MovieScene: Fixed shadowed variables

diff --git a/Engine/Source/Runtime/ActorAnimation/Private/ActorAnimation.cpp b/Engine/Source/Runtime/ActorAnimation/Private/ActorAnimation.cpp
--- a/Engine/Source/Runtime/ActorAnimation/Private/ActorAnimation.cpp
+++ b/Engine/Source/Runtime/ActorAnimation/Private/ActorAnimation.cpp
@@ -60,10 +60,10 @@ bool UActorAnimation::CanPossessObject(UObject& Object) const
 }
 
 
-void UActorAnimation::DestroyAllSpawnedObjects(UMovieScene& MovieScene)
+void UActorAnimation::DestroyAllSpawnedObjects(UMovieScene& SubMovieScene)
 {
 	const bool DestroyAll = true;
-	SpawnOrDestroyObjects(&MovieScene, DestroyAll);
+	SpawnOrDestroyObjects(&SubMovieScene, DestroyAll);
 	SpawnedActors.Empty();
 }
 
@@ -141,12 +141,12 @@ UObject* UActorAnimation::GetParentObject(UObject* Object) const
 }
 
 
-void UActorAnimation::SpawnOrDestroyObjects(UMovieScene* MovieScene, bool DestroyAll)
+void UActorAnimation::SpawnOrDestroyObjects(UMovieScene* SubMovieScene, bool DestroyAll)
 {
 	bool bAnyLevelActorsChanged = false;
 
 	// remove any proxy actors that we no longer need
-	if (MovieScene != nullptr)
+	if (SubMovieScene != nullptr)
 	{
 		for (auto SpawnedActorPair : SpawnedActors)
 		{
@@ -155,9 +155,9 @@ void UActorAnimation::SpawnOrDestroyObjects(UMovieScene* MovieScene, bool Destro
 			if (!DestroyAll)
 			{
 				// figure out if we still need this proxy actor
-				for (auto SpawnableIndex = 0; SpawnableIndex < MovieScene->GetSpawnableCount(); ++SpawnableIndex)
+				for (auto SpawnableIndex = 0; SpawnableIndex < SubMovieScene->GetSpawnableCount(); ++SpawnableIndex)
 				{
-					auto& Spawnable = MovieScene->GetSpawnable(SpawnableIndex);
+					auto& Spawnable = SubMovieScene->GetSpawnable(SpawnableIndex);
 
 					if (Spawnable.GetGuid() == SpawnedActorPair.Key)
 					{
@@ -210,11 +210,11 @@ void UActorAnimation::SpawnOrDestroyObjects(UMovieScene* MovieScene, bool Destro
 		}
 	}
 
-	if (!DestroyAll && (MovieScene != nullptr))
+	if (!DestroyAll && (SubMovieScene != nullptr))
 	{
-		for (auto SpawnableIndex = 0; SpawnableIndex < MovieScene->GetSpawnableCount(); ++SpawnableIndex)
+		for (auto SpawnableIndex = 0; SpawnableIndex < SubMovieScene->GetSpawnableCount(); ++SpawnableIndex)
 		{
-			FMovieSceneSpawnable& Spawnable = MovieScene->GetSpawnable( SpawnableIndex );
+			FMovieSceneSpawnable& Spawnable = SubMovieScene->GetSpawnable(SpawnableIndex);
 
 			// Do we already have a proxy for this spawnable?
 			if (SpawnedActors.Contains(Spawnable.GetGuid()))
diff --git a/Engine/Source/Runtime/ActorAnimation/Public/ActorAnimation.h b/Engine/Source/Runtime/ActorAnimation/Public/ActorAnimation.h
--- a/Engine/Source/Runtime/ActorAnimation/Public/ActorAnimation.h
+++ b/Engine/Source/Runtime/ActorAnimation/Public/ActorAnimation.h
@@ -27,12 +27,12 @@ public:
 	virtual bool AllowsSpawnableObjects() const override;
 	virtual void BindPossessableObject(const FGuid& ObjectId, UObject& PossessedObject) override;
 	virtual bool CanPossessObject(UObject& Object) const override;
-	virtual void DestroyAllSpawnedObjects(UMovieScene& MovieScene) override;
+	virtual void DestroyAllSpawnedObjects(UMovieScene& SubMovieScene) override;
 	virtual UObject* FindObject(const FGuid& ObjectId) const override;
 	virtual FGuid FindObjectId(UObject& Object) const override;
 	virtual UMovieScene* GetMovieScene() const override;
 	virtual UObject* GetParentObject(UObject* Object) const override;
-	virtual void SpawnOrDestroyObjects(UMovieScene* MovieScene, bool DestroyAll) override;
+	virtual void SpawnOrDestroyObjects(UMovieScene* SubMovieScene, bool DestroyAll) override;
 	virtual void UnbindPossessableObjects(const FGuid& ObjectId) override;
 
 #if WITH_EDITOR
diff --git a/Engine/Source/Runtime/MovieScene/Public/MovieSceneAnimation.h b/Engine/Source/Runtime/MovieScene/Public/MovieSceneAnimation.h
--- a/Engine/Source/Runtime/MovieScene/Public/MovieSceneAnimation.h
+++ b/Engine/Source/Runtime/MovieScene/Public/MovieSceneAnimation.h
@@ -45,7 +45,7 @@ class IMovieSceneAnimation
 	virtual void BindPossessableObject(const FGuid& ObjectId, UObject& PossessedObject) PURE_VIRTUAL(IMovieSceneAnimation::BindPossessableObject,);
 	
 	/**
-	 * Check whether the given object can be possessed by this object manager.
+	 * Check whether the given object can be possessed by this animation.
 	 *
 	 * @param Object The object to check.
 	 * @return true if the object can be possessed, false otherwise.
@@ -53,11 +53,11 @@ class IMovieSceneAnimation
 	virtual bool CanPossessObject(UObject& Object) const PURE_VIRTUAL(IMovieSceneAnimation::CanPossessObject, return false;);
 
 	/**
-	 * Destroy all objects that were spawned by this binding manager.
+	 * Destroy all objects that were spawned by this animation.
 	 *
-	 * @param MovieScene The movie scene that owns this manager.
+	 * @param SubMovieScene The (sub) movie scene to destroy objects for.
 	 */
-	virtual void DestroyAllSpawnedObjects(UMovieScene& MovieScene) PURE_VIRTUAL(IMovieSceneAnimation::DestroyAllSpawnedObjects,);
+	virtual void DestroyAllSpawnedObjects(UMovieScene& SubMovieScene) PURE_VIRTUAL(IMovieSceneAnimation::DestroyAllSpawnedObjects,);
 
 	/**
 	 * Gets the possessed or spawned object for the specified identifier.
@@ -96,10 +96,10 @@ class IMovieSceneAnimation
 	/**
 	 * Called to spawn or destroy objects for a movie instance.
 	 *
-	 * @param MovieScene The movie scene to spawn or destroy objects for.
+	 * @param SubMovieScene The (sub) movie scene to spawn or destroy objects for.
 	 * @param DestroyAll If true, destroy all spawned objects for the instance, if false only destroy unused objects.
 	 */
-	virtual void SpawnOrDestroyObjects(UMovieScene* MovieScene, bool DestroyAll) PURE_VIRTUAL(IMovieSceneAnimation::SpawnOrDestroyObjects,);
+	virtual void SpawnOrDestroyObjects(UMovieScene* SubMovieScene, bool DestroyAll) PURE_VIRTUAL(IMovieSceneAnimation::SpawnOrDestroyObjects,);
 
 	/**
 	 * Unbinds all possessable objects from the provided GUID.
diff --git a/Engine/Source/Runtime/Niagara/Private/NiagaraAnimation.cpp b/Engine/Source/Runtime/Niagara/Private/NiagaraAnimation.cpp
--- a/Engine/Source/Runtime/Niagara/Private/NiagaraAnimation.cpp
+++ b/Engine/Source/Runtime/Niagara/Private/NiagaraAnimation.cpp
@@ -33,7 +33,7 @@ bool UNiagaraAnimation::CanPossessObject(UObject& Object) const
 }
 
 
-void UNiagaraAnimation::DestroyAllSpawnedObjects(UMovieScene& MovieScene)
+void UNiagaraAnimation::DestroyAllSpawnedObjects(UMovieScene& SubMovieScene)
 {
 }
 
@@ -62,7 +62,7 @@ UObject* UNiagaraAnimation::GetParentObject(UObject* Object) const
 }
 
 
-void UNiagaraAnimation::SpawnOrDestroyObjects(UMovieScene* MovieScene, bool DestroyAll)
+void UNiagaraAnimation::SpawnOrDestroyObjects(UMovieScene* SubMovieScene, bool DestroyAll)
 {
 }
 
diff --git a/Engine/Source/Runtime/Niagara/Public/NiagaraAnimation.h b/Engine/Source/Runtime/Niagara/Public/NiagaraAnimation.h
--- a/Engine/Source/Runtime/Niagara/Public/NiagaraAnimation.h
+++ b/Engine/Source/Runtime/Niagara/Public/NiagaraAnimation.h
@@ -23,12 +23,12 @@ public:
 	virtual bool AllowsSpawnableObjects() const override;
 	virtual void BindPossessableObject(const FGuid& ObjectId, UObject& PossessedObject) override;
 	virtual bool CanPossessObject(UObject& Object) const override;
-	virtual void DestroyAllSpawnedObjects(UMovieScene& MovieScene) override;
+	virtual void DestroyAllSpawnedObjects(UMovieScene& SubMovieScene) override;
 	virtual UObject* FindObject(const FGuid& ObjectId) const override;
 	virtual FGuid FindObjectId(UObject& Object) const override;
 	virtual UMovieScene* GetMovieScene() const override;
 	virtual UObject* GetParentObject(UObject* Object) const override;
-	virtual void SpawnOrDestroyObjects(UMovieScene* MovieScene, bool DestroyAll) override;
+	virtual void SpawnOrDestroyObjects(UMovieScene* SubMovieScene, bool DestroyAll) override;
 	virtual void UnbindPossessableObjects(const FGuid& ObjectId) override;
 
 #if WITH_EDITOR
diff --git a/Engine/Source/Runtime/UMG/Public/Animation/WidgetAnimation.h b/Engine/Source/Runtime/UMG/Public/Animation/WidgetAnimation.h
--- a/Engine/Source/Runtime/UMG/Public/Animation/WidgetAnimation.h
+++ b/Engine/Source/Runtime/UMG/Public/Animation/WidgetAnimation.h
@@ -61,12 +61,12 @@ public:
 	virtual bool AllowsSpawnableObjects() const override;
 	virtual void BindPossessableObject(const FGuid& ObjectId, UObject& PossessedObject) override;
 	virtual bool CanPossessObject(UObject& Object) const override;
-	virtual void DestroyAllSpawnedObjects(UMovieScene& MovieScene) override { }
+	virtual void DestroyAllSpawnedObjects(UMovieScene& SubMovieScene) override { }
 	virtual UObject* FindObject(const FGuid& ObjectId) const override;
 	virtual FGuid FindObjectId(UObject& Object) const override;
 	virtual UMovieScene* GetMovieScene() const override;
 	virtual UObject* GetParentObject(UObject* Object) const override;
-	virtual void SpawnOrDestroyObjects(UMovieScene* MovieScene, bool DestroyAll) override { }
+	virtual void SpawnOrDestroyObjects(UMovieScene* SubMovieScene, bool DestroyAll) override { }
 	virtual void UnbindPossessableObjects(const FGuid& ObjectId) override;
 
 #if WITH_EDITOR
