Sequencer - Fix adding components to actor animations.

diff --git a/Engine/Source/Runtime/ActorAnimation/Private/ActorAnimation.cpp b/Engine/Source/Runtime/ActorAnimation/Private/ActorAnimation.cpp
--- a/Engine/Source/Runtime/ActorAnimation/Private/ActorAnimation.cpp
+++ b/Engine/Source/Runtime/ActorAnimation/Private/ActorAnimation.cpp
@@ -36,13 +36,12 @@ bool UActorAnimation::AllowsSpawnableObjects() const
 
 void UActorAnimation::BindPossessableObject(const FGuid& ObjectId, UObject& PossessedObject)
 {
-	FString ComponentName;
 	UActorComponent* Component = Cast<UActorComponent>(&PossessedObject);
 
 	if (Component != nullptr)
 	{
 		// @todo sequencer: gmp: need support for UStruct keys in TMap
-		PossessedObjects.Add(ObjectId.ToString(), FActorAnimationObject(PossessedObject.GetOuter(), ComponentName));
+		PossessedObjects.Add(ObjectId.ToString(), FActorAnimationObject(PossessedObject.GetOuter(), Component->GetName()));
 		//PossessedObjects.Add(ObjectId, FSequencerPossessedObject(PossessedObject.GetOuter(), ComponentName));
 	}
 	else
@@ -56,7 +55,7 @@ void UActorAnimation::BindPossessableObject(const FGuid& ObjectId, UObject& Poss
 
 bool UActorAnimation::CanPossessObject(UObject& Object) const
 {
-	return Object.IsA<AActor>();
+	return Object.IsA<AActor>() || Object.IsA<UActorComponent>();
 }
 
 
