Moved ActorAnimationPlayer to Runtime/ActorAnimation

Also changed the PlayActorAnimation function to CreateActorAnimationPlayer. PlayActorAnimation can be implemented in terms of CreateActorAnimationPlayer and UActorAnimationPlayer::Play()

diff --git a/Engine/Plugins/MovieScene/ActorAnimationEditor/Source/ActorAnimationEditor/Private/ActorAnimationPlayer.cpp b/Engine/Plugins/MovieScene/ActorAnimationEditor/Source/ActorAnimationEditor/Private/ActorAnimationPlayer.cpp
deleted file mode 100644
diff --git a/Engine/Plugins/MovieScene/ActorAnimationEditor/Source/ActorAnimationEditor/Private/ActorAnimationPlayer.h b/Engine/Plugins/MovieScene/ActorAnimationEditor/Source/ActorAnimationEditor/Private/ActorAnimationPlayer.h
deleted file mode 100644
diff --git a/Engine/Source/Runtime/ActorAnimation/Private/ActorAnimationPlayer.cpp b/Engine/Source/Runtime/ActorAnimation/Private/ActorAnimationPlayer.cpp
new file mode 100644
--- /dev/null
+++ b/Engine/Source/Runtime/ActorAnimation/Private/ActorAnimationPlayer.cpp
@@ -0,0 +1,271 @@
+// Copyright 1998-2015 Epic Games, Inc. All Rights Reserved.
+
+#include "ActorAnimationPrivatePCH.h"
+#include "ActorAnimationPlayer.h"
+#include "MovieScene.h"
+#include "MovieSceneSequence.h"
+#include "MovieSceneSequenceInstance.h"
+
+
+/* UActorAnimationPlayer structors
+ *****************************************************************************/
+
+UActorAnimationPlayer::UActorAnimationPlayer(const FObjectInitializer& ObjectInitializer)
+	: Super(ObjectInitializer)
+	, ActorAnimation(nullptr)
+	, bIsPlaying(false)
+	, TimeCursorPosition(0.0f)
+{ }
+
+
+/* UActorAnimationPlayer interface
+ *****************************************************************************/
+
+UActorAnimationPlayer* UActorAnimationPlayer::CreateActorAnimationPlayer(UObject* WorldContextObject, UActorAnimation* ActorAnimation)
+{
+	if (ActorAnimation == nullptr)
+	{
+		return nullptr;
+	}
+
+	UWorld* World = GEngine->GetWorldFromContextObject(WorldContextObject);
+
+	check(World != nullptr);
+	check(World->PersistentLevel != nullptr);
+
+	UActorAnimationPlayer* NewPlayer = NewObject<UActorAnimationPlayer>((UObject*)GetTransientPackage(), NAME_None, RF_Transient);
+
+	check(NewPlayer != nullptr);
+
+	NewPlayer->Initialize(ActorAnimation, World);
+	World->PersistentLevel->AddActiveRuntimeMovieScenePlayer(NewPlayer);
+
+	return NewPlayer;
+}
+
+
+bool UActorAnimationPlayer::IsPlaying() const
+{
+	return bIsPlaying;
+}
+
+
+void UActorAnimationPlayer::Pause()
+{
+	bIsPlaying = false;
+}
+
+
+void UActorAnimationPlayer::Play()
+{
+	if ((ActorAnimation == nullptr) || !World.IsValid())
+	{
+		return;
+	}
+
+	// @todo Sequencer playback: Should we recreate the instance every time?
+	RootMovieSceneInstance = MakeShareable(new FMovieSceneSequenceInstance(*ActorAnimation));
+
+	// @odo Sequencer Should we spawn actors here?
+	SpawnActorsForMovie(RootMovieSceneInstance.ToSharedRef());
+	RootMovieSceneInstance->RefreshInstance(*this);
+	
+	bIsPlaying = true;
+}
+
+
+/* UActorAnimationPlayer implementation
+ *****************************************************************************/
+
+void UActorAnimationPlayer::Initialize(UActorAnimation* InActorAnimation, UWorld* InWorld)
+{
+	ActorAnimation = InActorAnimation;
+	World = InWorld;
+}
+
+
+/* IMovieScenePlayer interface
+ *****************************************************************************/
+
+void UActorAnimationPlayer::SpawnActorsForMovie(TSharedRef<FMovieSceneSequenceInstance> MovieSceneInstance)
+{
+	UWorld* WorldPtr = World.Get();
+
+	if (WorldPtr == nullptr)
+	{
+		return;
+	}
+
+	UMovieScene* MovieScene = MovieSceneInstance->GetSequence()->GetMovieScene();
+
+	if (MovieScene == nullptr)
+	{
+		return;
+	}
+
+	TArray<FSpawnedActorInfo>* FoundSpawnedActors = InstanceToSpawnedActorMap.Find(MovieSceneInstance);
+
+	if (FoundSpawnedActors != nullptr)
+	{
+		// Remove existing spawned actors for this movie
+		DestroyActorsForMovie( MovieSceneInstance );
+	}
+
+	TArray<FSpawnedActorInfo>& SpawnedActorList = InstanceToSpawnedActorMap.Add(MovieSceneInstance, TArray<FSpawnedActorInfo>());
+
+	for (auto SpawnableIndex = 0; SpawnableIndex < MovieScene->GetSpawnableCount(); ++SpawnableIndex)
+	{
+		auto& Spawnable = MovieScene->GetSpawnable(SpawnableIndex);
+		UClass* GeneratedClass = Spawnable.GetClass();
+		
+		if ((GeneratedClass == nullptr) || !GeneratedClass->IsChildOf(AActor::StaticClass()))
+		{
+			continue;
+		}
+
+		AActor* ActorCDO = CastChecked<AActor>(GeneratedClass->ClassDefaultObject);
+		const FVector SpawnLocation = ActorCDO->GetRootComponent()->RelativeLocation;
+		const FRotator SpawnRotation = ActorCDO->GetRootComponent()->RelativeRotation;
+
+		FActorSpawnParameters SpawnInfo;
+		{
+			SpawnInfo.ObjectFlags = RF_NoFlags;
+		}
+
+		AActor* NewActor = WorldPtr->SpawnActor(GeneratedClass, &SpawnLocation, &SpawnRotation, SpawnInfo);
+
+		if (NewActor)
+		{	
+			// Actor was spawned OK!
+			FSpawnedActorInfo NewInfo;
+			{
+				NewInfo.RuntimeGuid = Spawnable.GetGuid();
+				NewInfo.SpawnedActor = NewActor;
+			}
+						
+			SpawnedActorList.Add(NewInfo);
+		}
+	}
+}
+
+
+void UActorAnimationPlayer::DestroyActorsForMovie(TSharedRef<FMovieSceneSequenceInstance> MovieSceneInstance)
+{
+	UWorld* WorldPtr = World.Get();
+
+	if (WorldPtr == nullptr)
+	{
+		return;
+	}
+
+	TArray<FSpawnedActorInfo>* SpawnedActors = InstanceToSpawnedActorMap.Find(MovieSceneInstance);
+
+	if (SpawnedActors == nullptr)
+	{
+		return;
+	}
+
+	TArray<FSpawnedActorInfo>& SpawnedActorsRef = *SpawnedActors;
+
+	for(int32 ActorIndex = 0; ActorIndex < SpawnedActors->Num(); ++ActorIndex)
+	{
+		AActor* Actor = SpawnedActorsRef[ActorIndex].SpawnedActor.Get();
+		
+		if (Actor != nullptr)
+		{
+			// @todo Sequencer figure this out.  Defaults to false.
+			const bool bNetForce = false;
+
+			// At runtime, level modification is not needed
+			const bool bShouldModifyLevel = false;
+			Actor->GetWorld()->DestroyActor(Actor, bNetForce, bShouldModifyLevel);
+		}
+	}
+
+	InstanceToSpawnedActorMap.Remove(MovieSceneInstance);
+}
+
+
+void UActorAnimationPlayer::GetRuntimeObjects(TSharedRef<FMovieSceneSequenceInstance> MovieSceneInstance, const FGuid& ObjectId, TArray<UObject*>& OutObjects) const
+{
+	// @todo sequencer runtime: Add support for individually spawning actors on demand when first requested?
+	//    This may be important to reduce the up-front hitch when spawning actors for the entire movie, but
+	//    may introduce smaller hitches during playback.  Needs experimentation.
+
+	const TArray<FSpawnedActorInfo>* SpawnedActors = InstanceToSpawnedActorMap.Find(MovieSceneInstance);
+
+	if ((SpawnedActors != nullptr) && (SpawnedActors->Num() > 0))
+	{
+		// we have a spawned actor for this handle
+		const TArray<FSpawnedActorInfo>& SpawnedActorInfoRef = *SpawnedActors;
+
+		for (int32 SpawnedActorIndex = 0; SpawnedActorIndex < SpawnedActorInfoRef.Num(); ++SpawnedActorIndex)
+		{
+			const FSpawnedActorInfo ActorInfo = SpawnedActorInfoRef[SpawnedActorIndex];
+
+			if ((ActorInfo.RuntimeGuid == ObjectId) && ActorInfo.SpawnedActor.IsValid())
+			{
+				OutObjects.Add(ActorInfo.SpawnedActor.Get());
+			}
+		}
+	}
+	else
+	{
+		// otherwise, check whether we have one or more possessed actors that are mapped to this handle
+		UObject* FoundObject = ActorAnimation->FindObject(ObjectId);
+
+		if (FoundObject != nullptr)
+		{
+			OutObjects.Add(FoundObject);
+		}
+	}
+}
+
+
+void UActorAnimationPlayer::UpdateCameraCut(UObject* ObjectToViewThrough, bool bNewCameraCut) const
+{
+}
+
+
+EMovieScenePlayerStatus::Type UActorAnimationPlayer::GetPlaybackStatus() const
+{
+	return bIsPlaying ? EMovieScenePlayerStatus::Playing : EMovieScenePlayerStatus::Stopped;
+}
+
+
+void UActorAnimationPlayer::AddOrUpdateMovieSceneInstance(UMovieSceneSection& MovieSceneSection, TSharedRef<FMovieSceneSequenceInstance> InstanceToAdd)
+{
+	SpawnActorsForMovie( InstanceToAdd );
+}
+
+
+void UActorAnimationPlayer::RemoveMovieSceneInstance(UMovieSceneSection& MovieSceneSection, TSharedRef<FMovieSceneSequenceInstance> InstanceToRemove)
+{
+	const bool bDestroyAll = true;
+	DestroyActorsForMovie( InstanceToRemove );
+}
+
+
+TSharedRef<FMovieSceneSequenceInstance> UActorAnimationPlayer::GetRootMovieSceneSequenceInstance() const
+{
+	return RootMovieSceneInstance.ToSharedRef();
+}
+
+
+/* IRuntimeMovieScenePlayerInterface interface
+ *****************************************************************************/
+
+void UActorAnimationPlayer::Tick(const float DeltaSeconds)
+{
+	float LastTimePosition = TimeCursorPosition;
+
+	if (bIsPlaying)
+	{
+		TimeCursorPosition += DeltaSeconds;
+	}
+
+	if(RootMovieSceneInstance.IsValid())
+	{
+		RootMovieSceneInstance->Update(TimeCursorPosition, LastTimePosition, *this);
+	}
+}
diff --git a/Engine/Source/Runtime/ActorAnimation/Public/ActorAnimationPlayer.h b/Engine/Source/Runtime/ActorAnimation/Public/ActorAnimationPlayer.h
new file mode 100644
--- /dev/null
+++ b/Engine/Source/Runtime/ActorAnimation/Public/ActorAnimationPlayer.h
@@ -0,0 +1,114 @@
+// Copyright 1998-2015 Epic Games, Inc. All Rights Reserved.
+
+#pragma once
+
+#include "IMovieScenePlayer.h"
+#include "ActorAnimation.h"
+#include "MovieScene/RuntimeMovieScenePlayerInterface.h"
+#include "ActorAnimationPlayer.generated.h"
+
+
+class FMovieSceneSequenceInstance;
+class ULevel;
+class UMovieSceneBindings;
+
+
+/**
+ * UActorAnimationPlayer is used to actually "play" an actor animation asset at runtime.
+ *
+ * This class keeps track of playback state and provides functions for manipulating
+ * an actor animation while its playing.
+ */
+UCLASS(BlueprintType)
+class ACTORANIMATION_API UActorAnimationPlayer
+	: public UObject
+	, public IMovieScenePlayer
+	, public IRuntimeMovieScenePlayerInterface
+{
+	GENERATED_UCLASS_BODY()
+
+public:
+
+	/**
+	 * Create a new actor animation player.
+	 *
+	 * @param WorldContextObject Context object from which to retrieve a UWorld.
+	 * @param ActorAnimation The actor animation to play.
+	 */
+	UFUNCTION(BlueprintCallable, Category="Game|Cinematic", meta=(WorldContext="WorldContextObject"))
+	static UActorAnimationPlayer* CreateActorAnimationPlayer(UObject* WorldContextObject, UActorAnimation* ActorAnimation);
+
+	/** Start playback from the current time cursor position. */
+	UFUNCTION(BlueprintCallable, Category="Game|Cinematic")
+	void Play();
+
+	/** Pause playback. */
+	UFUNCTION(BlueprintCallable, Category="Game|Cinematic")
+	void Pause();
+
+	/** Check whether the sequence is actively playing. */
+	bool IsPlaying() const;
+
+protected:
+
+	/**
+	 * Initialize the player.
+	 *
+	 * @param InActorAnimation The actor animation to play.
+	 * @param InWorld The world that the animation is played in.
+	 */
+	void Initialize(UActorAnimation* InActorAnimation, UWorld* InWorld);
+
+protected:
+
+	// IMovieScenePlayer interface
+
+	virtual void SpawnActorsForMovie(TSharedRef<FMovieSceneSequenceInstance> MovieSceneInstance);
+	virtual void DestroyActorsForMovie(TSharedRef<FMovieSceneSequenceInstance> MovieSceneInstance);
+	virtual void GetRuntimeObjects(TSharedRef<FMovieSceneSequenceInstance> MovieSceneInstance, const FGuid& ObjectHandle, TArray<UObject*>& OutObjects) const override;
+	virtual void UpdateCameraCut(UObject* ObjectToViewThrough, bool bNewCameraCut) const override;
+	virtual EMovieScenePlayerStatus::Type GetPlaybackStatus() const override;
+	virtual void AddOrUpdateMovieSceneInstance(UMovieSceneSection& MovieSceneSection, TSharedRef<FMovieSceneSequenceInstance> InstanceToAdd) override;
+	virtual void RemoveMovieSceneInstance(UMovieSceneSection& MovieSceneSection, TSharedRef<FMovieSceneSequenceInstance> InstanceToRemove) override;
+	virtual TSharedRef<FMovieSceneSequenceInstance> GetRootMovieSceneSequenceInstance() const override;
+
+protected:
+
+	// IRuntimeMovieScenePlayerInterface
+
+	virtual void Tick(const float DeltaSeconds) override;
+
+private:
+
+	/** The actor animation to play. */
+	UPROPERTY()
+	UActorAnimation* ActorAnimation;
+
+	/** Whether we're currently playing. If false, then sequence playback is paused or was never started. */
+	UPROPERTY()
+	bool bIsPlaying;
+
+	/** The current time cursor position within the sequence (in seconds) */
+	UPROPERTY()
+	double TimeCursorPosition;
+
+private:
+
+	struct FSpawnedActorInfo
+	{
+		/** Identifier that maps this actor to a movie scene spawnable */
+		FGuid RuntimeGuid;
+
+		/** The spawned actor */
+		TWeakObjectPtr<AActor> SpawnedActor;
+	};
+
+	/** Maps spawnable GUIDs to their spawned actor in the world */
+	TMap<TWeakPtr<FMovieSceneSequenceInstance>, TArray<FSpawnedActorInfo>> InstanceToSpawnedActorMap;
+
+	/** The root movie scene instance to update when playing. */
+	TSharedPtr<FMovieSceneSequenceInstance> RootMovieSceneInstance;
+
+	/** The world this player will spawn actors in, if needed */
+	TWeakObjectPtr<UWorld> World;
+};
