Sequencer - Cache components with a weak pointer to avoid leaks when changing levels with the sequence still open.

diff --git a/Engine/Source/Runtime/ActorAnimation/Private/ActorAnimationObject.cpp b/Engine/Source/Runtime/ActorAnimation/Private/ActorAnimationObject.cpp
--- a/Engine/Source/Runtime/ActorAnimation/Private/ActorAnimationObject.cpp
+++ b/Engine/Source/Runtime/ActorAnimation/Private/ActorAnimationObject.cpp
@@ -16,9 +16,9 @@ UObject* FActorAnimationObject::GetObject() const
 	}
 
 	// get cached component
-	if ((CachedComponent != nullptr) && (CachedComponent->GetName() == ComponentName))
+	if ((CachedComponent.IsValid()) && (CachedComponent->GetName() == ComponentName))
 	{
-		return CachedComponent;
+		return CachedComponent.Get();
 	}
 
 	CachedComponent = nullptr;
diff --git a/Engine/Source/Runtime/ActorAnimation/Public/ActorAnimationObject.h b/Engine/Source/Runtime/ActorAnimation/Public/ActorAnimationObject.h
--- a/Engine/Source/Runtime/ActorAnimation/Public/ActorAnimationObject.h
+++ b/Engine/Source/Runtime/ActorAnimation/Public/ActorAnimationObject.h
@@ -74,5 +74,5 @@ private:
 
 	/** Cached pointer to the Actor component (only if ComponentName is set). */
 	UPROPERTY(transient)
-	mutable UObject* CachedComponent;
+	mutable TWeakObjectPtr<UObject> CachedComponent;
 };
