Added MovieSceneActor, overhauled UActorAnimationPlayer, and fixed Actor Animation object bindings

Lots of change in here:
 - Improved UActorAnimationPlayer functionality to include looping and play speed (UI needs polish)
 - Fixed actor object bindings not working in PIE. We now only store the object GUID (and an object path for a fallback), which can then be fixed up within a given context (UWorld for now).
 - Added UActorAnimationInstance which is responsible for managing bindings on a UActorAnimation. UActorAnimation itself is now pretty much just a data asset.
 - Levels are no longer responsible for ticking actor animations, this is now handled by the client (either AMovieSceneActor, or ticked automatically if created from a blueprint node)

diff --git a/Engine/Plugins/MovieScene/ActorAnimationEditor/Source/ActorAnimationEditor/Private/ActorAnimationEditorToolkit.cpp b/Engine/Plugins/MovieScene/ActorAnimationEditor/Source/ActorAnimationEditor/Private/ActorAnimationEditorToolkit.cpp
--- a/Engine/Plugins/MovieScene/ActorAnimationEditor/Source/ActorAnimationEditor/Private/ActorAnimationEditorToolkit.cpp
+++ b/Engine/Plugins/MovieScene/ActorAnimationEditor/Source/ActorAnimationEditor/Private/ActorAnimationEditorToolkit.cpp
@@ -66,7 +66,7 @@ FActorAnimationEditorToolkit::~FActorAnimationEditorToolkit()
 /* FActorAnimationEditorToolkit interface
  *****************************************************************************/
 
-void FActorAnimationEditorToolkit::Initialize( const EToolkitMode::Type Mode, const TSharedPtr<IToolkitHost>& InitToolkitHost, UActorAnimation* ActorAnimation, bool bEditWithinLevelEditor )
+void FActorAnimationEditorToolkit::Initialize( const EToolkitMode::Type Mode, const TSharedPtr<IToolkitHost>& InitToolkitHost, UActorAnimationInstance* InActorAnimationInstance, bool bEditWithinLevelEditor )
 {
 	// create tab layout
 	const TSharedRef<FTabManager::FLayout> StandaloneDefaultLayout = FTabManager::NewLayout("Standalone_ActorAnimationEditor")
@@ -80,10 +80,12 @@ void FActorAnimationEditorToolkit::Initialize( const EToolkitMode::Type Mode, co
 				)
 		);
 
+	ActorAnimationInstance = InActorAnimationInstance;
+
 	const bool bCreateDefaultStandaloneMenu = true;
 	const bool bCreateDefaultToolbar = false;
 
-	FAssetEditorToolkit::InitAssetEditor(Mode, InitToolkitHost, SequencerDefs::SequencerAppIdentifier, StandaloneDefaultLayout, bCreateDefaultStandaloneMenu, bCreateDefaultToolbar, ActorAnimation);
+	FAssetEditorToolkit::InitAssetEditor(Mode, InitToolkitHost, SequencerDefs::SequencerAppIdentifier, StandaloneDefaultLayout, bCreateDefaultStandaloneMenu, bCreateDefaultToolbar, ActorAnimationInstance);
 
 	// initialize sequencer
 	FSequencerViewParams ViewParams(TEXT("ActorAnimationEditorViewParams"));
@@ -94,7 +96,7 @@ void FActorAnimationEditorToolkit::Initialize( const EToolkitMode::Type Mode, co
 	FSequencerInitParams SequencerInitParams;
 	{
 		SequencerInitParams.ViewParams = ViewParams;
-		SequencerInitParams.RootSequence = ActorAnimation;
+		SequencerInitParams.RootSequence = ActorAnimationInstance;
 		SequencerInitParams.bEditWithinLevelEditor = bEditWithinLevelEditor;
 		SequencerInitParams.ToolkitHost = InitToolkitHost;
 	}
@@ -235,6 +237,8 @@ void FActorAnimationEditorToolkit::HandleAddMenuAddSlomoTrackExecute()
 
 void FActorAnimationEditorToolkit::HandleMapChanged(class UWorld* NewWorld, EMapChangeType MapChangeType)
 {
+	ActorAnimationInstance->SetContext(NewWorld);
+
 	Sequencer->NotifyMapChanged(NewWorld, MapChangeType);
 }
 
diff --git a/Engine/Plugins/MovieScene/ActorAnimationEditor/Source/ActorAnimationEditor/Private/ActorAnimationEditorToolkit.h b/Engine/Plugins/MovieScene/ActorAnimationEditor/Source/ActorAnimationEditor/Private/ActorAnimationEditorToolkit.h
--- a/Engine/Plugins/MovieScene/ActorAnimationEditor/Source/ActorAnimationEditor/Private/ActorAnimationEditorToolkit.h
+++ b/Engine/Plugins/MovieScene/ActorAnimationEditor/Source/ActorAnimationEditor/Private/ActorAnimationEditorToolkit.h
@@ -3,7 +3,7 @@
 #pragma once
 
 #include "Toolkits/AssetEditorToolkit.h"
-
+#include "ActorAnimation.h"
 
 enum class EMapChangeType : uint8;
 class FTabManager;
@@ -11,7 +11,6 @@ class ILevelViewport;
 class ISequencer;
 class IToolkitHost;
 class SWidget;
-class UActorAnimation;
 class UWorld;
 struct FSequencerViewParams;
 
@@ -21,6 +20,7 @@ struct FSequencerViewParams;
  */
 class FActorAnimationEditorToolkit
 	: public FAssetEditorToolkit
+	, public FGCObject
 { 
 public:
 
@@ -45,7 +45,7 @@ public:
 	 * @param TrackEditorDelegates Delegates to call to create auto-key handlers for this sequencer.
 	 * @param bEditWithinLevelEditor Whether or not sequencer should be edited within the level editor.
 	 */
-	void Initialize(const EToolkitMode::Type Mode, const TSharedPtr<IToolkitHost>& InitToolkitHost, UActorAnimation* ActorAnimation, bool bEditWithinLevelEditor);
+	void Initialize(const EToolkitMode::Type Mode, const TSharedPtr<IToolkitHost>& InitToolkitHost, UActorAnimationInstance* ActorAnimationInstance, bool bEditWithinLevelEditor);
 
 public:
 
@@ -60,6 +60,11 @@ public:
 
 	void UpdateViewports(AActor* ActorToViewThrough) const;
 
+	void AddReferencedObjects( FReferenceCollector& Collector ) override
+	{
+		Collector.AddReferencedObject(ActorAnimationInstance);
+	}
+
 private:
 
 	/** Callback for executing the Add Component action. */
@@ -105,6 +110,9 @@ private:
 	/** Pointer to the style set to use for toolkits. */
 	TSharedRef<ISlateStyle> Style;
 
+	/** Actor animation instance for our edit operation */
+	UActorAnimationInstance* ActorAnimationInstance;
+
 private:
 
 	/**	The tab ids for all the tabs used */
diff --git a/Engine/Plugins/MovieScene/ActorAnimationEditor/Source/ActorAnimationEditor/Private/AssetTools/ActorAnimationActions.cpp b/Engine/Plugins/MovieScene/ActorAnimationEditor/Source/ActorAnimationEditor/Private/AssetTools/ActorAnimationActions.cpp
--- a/Engine/Plugins/MovieScene/ActorAnimationEditor/Source/ActorAnimationEditor/Private/AssetTools/ActorAnimationActions.cpp
+++ b/Engine/Plugins/MovieScene/ActorAnimationEditor/Source/ActorAnimationEditor/Private/AssetTools/ActorAnimationActions.cpp
@@ -46,6 +46,21 @@ FColor FActorAnimationActions::GetTypeColor() const
 
 void FActorAnimationActions::OpenAssetEditor(const TArray<UObject*>& InObjects, TSharedPtr<IToolkitHost> EditWithinLevelEditor)
 {
+	UWorld* WorldContext = nullptr;
+	for (const FWorldContext& Context : GEngine->GetWorldContexts())
+	{
+		if (Context.WorldType == EWorldType::Editor)
+		{
+			WorldContext = Context.World();
+			break;
+		}
+	}
+
+	if (!ensure(WorldContext))
+	{
+		return;
+	}
+
 	EToolkitMode::Type Mode = EditWithinLevelEditor.IsValid()
 		? EToolkitMode::WorldCentric
 		: EToolkitMode::Standalone;
@@ -56,8 +71,16 @@ void FActorAnimationActions::OpenAssetEditor(const TArray<UObject*>& InObjects,
 
 		if (ActorAnimation != nullptr)
 		{
+			// Legacy upgrade
+			ActorAnimation->ConvertPersistentBindingsToDefault(WorldContext);
+
+			// Create an edit instance for this actor animation that can only edit the default bindings in the current world
+			UActorAnimationInstance* Instance = NewObject<UActorAnimationInstance>(GetTransientPackage());
+			bool bCanInstanceBindings = false;
+			Instance->Initialize(ActorAnimation, WorldContext, bCanInstanceBindings);
+
 			TSharedRef<FActorAnimationEditorToolkit> Toolkit = MakeShareable(new FActorAnimationEditorToolkit(Style));
-			Toolkit->Initialize(Mode, EditWithinLevelEditor, ActorAnimation, true);
+			Toolkit->Initialize(Mode, EditWithinLevelEditor, Instance, true);
 		}
 	}
 }
diff --git a/Engine/Source/Editor/EditorStyle/Private/SlateEditorStyle.cpp b/Engine/Source/Editor/EditorStyle/Private/SlateEditorStyle.cpp
--- a/Engine/Source/Editor/EditorStyle/Private/SlateEditorStyle.cpp
+++ b/Engine/Source/Editor/EditorStyle/Private/SlateEditorStyle.cpp
@@ -5332,6 +5332,7 @@ void FSlateEditorStyle::FStyle::SetupClassIconsAndThumbnails()
 			TEXT("MediaPlayer"),
 			TEXT("MediaSoundWave"),
 			TEXT("MediaTexture"),
+			TEXT("MovieSceneActor"),
 			TEXT("MultiFont"),
 			TEXT("NavLinkProxy"),
 			TEXT("NavMeshBoundsVolume"),
diff --git a/Engine/Source/Editor/UnrealEd/Private/Factories/ActorFactory.cpp b/Engine/Source/Editor/UnrealEd/Private/Factories/ActorFactory.cpp
--- a/Engine/Source/Editor/UnrealEd/Private/Factories/ActorFactory.cpp
+++ b/Engine/Source/Editor/UnrealEd/Private/Factories/ActorFactory.cpp
@@ -67,6 +67,10 @@ ActorFactory.cpp:
 #include "Kismet2/ComponentEditorUtils.h"
 #include "Components/BillboardComponent.h"
 
+#include "ActorAnimation.h"
+#include "MovieSceneActor.h"
+#include "ActorFactoryMovieScene.h"
+
 DEFINE_LOG_CATEGORY_STATIC(LogActorFactory, Log, All);
 
 #define LOCTEXT_NAMESPACE "ActorFactory"
@@ -1939,4 +1943,55 @@ void UActorFactoryCylinderVolume::PostSpawnActor( UObject* Asset, AActor* NewAct
 	}
 }
 
+/*-----------------------------------------------------------------------------
+UActorFactoryMovieScene
+-----------------------------------------------------------------------------*/
+UActorFactoryMovieScene::UActorFactoryMovieScene(const FObjectInitializer& ObjectInitializer)
+	: Super(ObjectInitializer)
+{
+	DisplayName = LOCTEXT("MovieSceneDisplayName", "MovieScene");
+	NewActorClass = AMovieSceneActor::StaticClass();
+}
+
+bool UActorFactoryMovieScene::CanCreateActorFrom( const FAssetData& AssetData, FText& OutErrorMsg )
+{
+	if ( UActorFactory::CanCreateActorFrom( AssetData, OutErrorMsg ) )
+	{
+		return true;
+	}
+
+	if ( AssetData.IsValid() && !AssetData.GetClass()->IsChildOf( UActorAnimation::StaticClass() ) )
+	{
+		OutErrorMsg = NSLOCTEXT("CanCreateActor", "NoActorAnimationAsset", "A valid actor animation asset must be specified.");
+		return false;
+	}
+
+	return true;
+}
+
+AActor* UActorFactoryMovieScene::SpawnActor( UObject* Asset, ULevel* InLevel, const FVector& Location, const FRotator& Rotation, EObjectFlags ObjectFlags, const FName& Name )
+{
+	AMovieSceneActor* NewActor = Cast<AMovieSceneActor>(Super::SpawnActor(Asset, InLevel, Location, Rotation, ObjectFlags, Name));
+
+	if (NewActor)
+	{
+		if (UActorAnimation* ActorAnimation = Cast<UActorAnimation>(Asset))
+		{
+			NewActor->SetAnimation(ActorAnimation);
+		}
+	}
+
+	return NewActor;
+}
+
+UObject* UActorFactoryMovieScene::GetAssetFromActorInstance(AActor* Instance)
+{
+	if (AMovieSceneActor* MovieSceneActor = Cast<AMovieSceneActor>(Instance))
+	{
+		return MovieSceneActor->ActorAnimation.TryLoad();
+	}
+
+	return nullptr;
+}
+
 #undef LOCTEXT_NAMESPACE
diff --git a/Engine/Source/Editor/UnrealEd/Private/Factories/ActorFactoryMovieScene.h b/Engine/Source/Editor/UnrealEd/Private/Factories/ActorFactoryMovieScene.h
new file mode 100644
--- /dev/null
+++ b/Engine/Source/Editor/UnrealEd/Private/Factories/ActorFactoryMovieScene.h
@@ -0,0 +1,21 @@
+// Copyright 1998-2015 Epic Games, Inc. All Rights Reserved.
+
+#pragma once
+
+#include "ActorFactories/ActorFactory.h"
+#include "ActorFactoryMovieScene.generated.h"
+
+UCLASS()
+class UActorFactoryMovieScene : public UActorFactory
+{
+	GENERATED_UCLASS_BODY()
+
+	// Begin UActorFactory Interface
+	virtual AActor* SpawnActor( UObject* Asset, ULevel* InLevel, const FVector& Location, const FRotator& Rotation, EObjectFlags ObjectFlags, const FName& Name ) override;
+	virtual bool CanCreateActorFrom( const FAssetData& AssetData, FText& OutErrorMsg ) override;
+	virtual UObject* GetAssetFromActorInstance(AActor* ActorInstance) override;
+	// End UActorFactory Interface
+};
+
+
+
diff --git a/Engine/Source/Editor/UnrealEd/UnrealEd.Build.cs b/Engine/Source/Editor/UnrealEd/UnrealEd.Build.cs
--- a/Engine/Source/Editor/UnrealEd/UnrealEd.Build.cs
+++ b/Engine/Source/Editor/UnrealEd/UnrealEd.Build.cs
@@ -92,6 +92,7 @@ public class UnrealEd : ModuleRules
 		PrivateDependencyModuleNames.AddRange(
 			new string[] 
 			{ 
+                "ActorAnimation",
 				"AnimGraph",
                 "AppFramework",
 				"BlueprintGraph",
@@ -136,7 +137,7 @@ public class UnrealEd : ModuleRules
                 "MaterialUtilities",
                 "LocalizationService",
 				"AddContentDialog",
-				"GameProjectGeneration"                
+				"GameProjectGeneration",                
 			}
 		);
 
diff --git a/Engine/Source/Runtime/ActorAnimation/Private/ActorAnimation.cpp b/Engine/Source/Runtime/ActorAnimation/Private/ActorAnimation.cpp
--- a/Engine/Source/Runtime/ActorAnimation/Private/ActorAnimation.cpp
+++ b/Engine/Source/Runtime/ActorAnimation/Private/ActorAnimation.cpp
@@ -6,17 +6,11 @@
 #include "MovieScene.h"
 
 
-/* UActorAnimation structors
- *****************************************************************************/
-
 UActorAnimation::UActorAnimation(const FObjectInitializer& ObjectInitializer)
 	: Super(ObjectInitializer)
 	, MovieScene(nullptr)
-{ }
-
-
-/* UActorAnimation interface
- *****************************************************************************/
+{
+}
 
 void UActorAnimation::Initialize()
 {
@@ -24,86 +18,162 @@ void UActorAnimation::Initialize()
 	MovieScene = NewObject<UMovieScene>(this, NAME_None, RF_Transactional);
 }
 
+void UActorAnimation::ConvertPersistentBindingsToDefault(UObject* FixupContext)
+{
+	if (PossessedObjects_DEPRECATED.Num() == 0)
+	{
+		return;
+	}
 
-/* UMovieSceneAnimation overrides
- *****************************************************************************/
+	MarkPackageDirty();
+	for (auto& Pair : PossessedObjects_DEPRECATED)
+	{
+		UObject* Object = Pair.Value.GetObject();
+		if (Object)
+		{
+			FGuid ObjectId;
+			FGuid::Parse(Pair.Key, ObjectId);
+			DefaultObjectReferences.CreateBinding(ObjectId, Object, FixupContext);
+		}
+	}
+	PossessedObjects_DEPRECATED.Empty();
+}
 
-bool UActorAnimation::AllowsSpawnableObjects() const
+void UActorAnimationInstance::Initialize(UActorAnimation* InAnimation, UObject* InContext, bool bInCanRemapBindings)
 {
-	return true;
+	ActorAnimation = InAnimation;
+	bCanRemapBindings = bInCanRemapBindings;
+
+	if (!bCanRemapBindings)
+	{
+		RemappedObjectReferences = FActorAnimationObjectReferenceMap();
+	}
+
+	SetContext(InContext);
 }
 
+void UActorAnimationInstance::SetContext(UObject* NewContext)
+{
+	ResolutionContext = NewContext;
+	CachedObjectBindings.Reset();
+}
 
-void UActorAnimation::BindPossessableObject(const FGuid& ObjectId, UObject& PossessedObject)
+bool UActorAnimationInstance::AllowsSpawnableObjects() const
 {
-	UActorComponent* Component = Cast<UActorComponent>(&PossessedObject);
+	return true;
+}
 
-	if (Component != nullptr)
+void UActorAnimationInstance::BindPossessableObject(const FGuid& ObjectId, UObject& PossessedObject)
+{
+	UObject* Context = ResolutionContext.Get();
+	if (!Context)
+	{
+		return;
+	}
+
+	if (!ActorAnimation)
 	{
-		// @todo sequencer: gmp: need support for UStruct keys in TMap
-		PossessedObjects.Add(ObjectId.ToString(), FActorAnimationObject(PossessedObject.GetOuter(), Component->GetName()));
-		//PossessedObjects.Add(ObjectId, FSequencerPossessedObject(PossessedObject.GetOuter(), ComponentName));
+		RemappedObjectReferences.CreateBinding(ObjectId, &PossessedObject, Context);
+	}
+	else if (bCanRemapBindings && ActorAnimation->DefaultObjectReferences.HasBinding(ObjectId))
+	{
+		RemappedObjectReferences.CreateBinding(ObjectId, &PossessedObject, Context);
 	}
 	else
 	{
-		// @todo sequencer: gmp: need support for UStruct keys in TMap
-		PossessedObjects.Add(ObjectId.ToString(), FActorAnimationObject(&PossessedObject));
-		//PossessedObjects.Add(ObjectId, FSequencerPossessedObject(Actor));
+		ActorAnimation->DefaultObjectReferences.CreateBinding(ObjectId, &PossessedObject, Context);
 	}
 }
 
-
-bool UActorAnimation::CanPossessObject(UObject& Object) const
+bool UActorAnimationInstance::CanPossessObject(UObject& Object) const
 {
 	return Object.IsA<AActor>() || Object.IsA<UActorComponent>();
 }
 
-
-void UActorAnimation::DestroyAllSpawnedObjects()
+void UActorAnimationInstance::DestroyAllSpawnedObjects()
 {
 	const bool DestroyAll = true;
 	SpawnOrDestroyObjects(DestroyAll);
 	SpawnedActors.Empty();
 }
 
-
-UObject* UActorAnimation::FindObject(const FGuid& ObjectId) const
+UObject* UActorAnimationInstance::FindObject(const FGuid& ObjectId) const
 {
-	// search bindings (possessables)
-	// @todo sequencer: gmp: need support for UStruct keys in TMap
-	const FActorAnimationObject* PossessedObject = PossessedObjects.Find(ObjectId.ToString());
-	//const FSequencerPossessedObject* PossessedObject = PossessedObjects.Find(ObjectId);
+	// If it's already cached, we can just return that
+	if (auto* WeakCachedObject = CachedObjectBindings.Find(ObjectId))
+	{
+		if (auto* CachedObj = WeakCachedObject->Get())
+		{
+			return CachedObj;
+		}
+	}
+
+	// Look up the object from scratch
+	UObject* FoundObject = nullptr;
 
-	if (PossessedObject != nullptr)
+	// Otherwise we need to try and find it
+	UObject* Context = ResolutionContext.Get();
+	if (!FoundObject && Context)
 	{
-		return PossessedObject->GetObject();
+		if (!FoundObject && bCanRemapBindings)
+		{
+			// Attempt to resolve the object binding through the remapped bindings
+			FoundObject = RemappedObjectReferences.ResolveBinding(ObjectId, Context);
+		}
+		
+		if (ActorAnimation)
+		{
+			// Attempt to resolve the object binding through the defaults
+			FoundObject = ActorAnimation->DefaultObjectReferences.ResolveBinding(ObjectId, Context);
+		}
 	}
 
-	// search proxies (spawnables)
-	TWeakObjectPtr<AActor> Actor = SpawnedActors.FindRef(ObjectId);
+	if (!FoundObject)
+	{
+		// Maybe it's a spawnable then?
+		FoundObject = SpawnedActors.FindRef(ObjectId).Get();
+	}
 
-	if (Actor.IsValid())
+	if (FoundObject)
 	{
-		return Actor.Get();
+		CachedObjectBindings.Add(ObjectId, FoundObject);
 	}
 
-	// not found
-	return nullptr;
+	return FoundObject;
 }
 
-
-FGuid UActorAnimation::FindObjectId(UObject& Object) const
+FGuid UActorAnimationInstance::FindObjectId(UObject& Object) const
 {
+	for (auto& Pair : CachedObjectBindings)
+	{
+		if (Pair.Value == &Object)
+		{
+			return Pair.Key;
+		}
+	}
+
 	// search possessed objects
-	for (auto PossessedObjectPair : PossessedObjects)
+	UObject* Context = ResolutionContext.Get();
+	if (Context)
 	{
-		const FActorAnimationObject& PossessedObject = PossessedObjectPair.Value;
+		if (bCanRemapBindings)
+		{
+			// Search instances
+			FGuid ObjectId = RemappedObjectReferences.FindBindingId(&Object, Context);
+			if (ObjectId.IsValid())
+			{
+				return ObjectId;
+			}
+		}
 
-		if (PossessedObject.GetObject() == &Object)
+		// Search defaults
+		if (ActorAnimation)
 		{
-			// @todo sequencer: gmp: need support for UStruct keys in TMap
-			FGuid Result; FGuid::Parse(PossessedObjectPair.Key, Result); return Result;
-			//return PossessedObjectsPair.Key;
+			FGuid ObjectId = ActorAnimation->DefaultObjectReferences.FindBindingId(&Object, Context);
+			if (ObjectId.IsValid())
+			{
+				return ObjectId;
+			}
 		}
 	}
 
@@ -120,14 +190,12 @@ FGuid UActorAnimation::FindObjectId(UObject& Object) const
 	return FGuid();
 }
 
-
-UMovieScene* UActorAnimation::GetMovieScene() const
+UMovieScene* UActorAnimationInstance::GetMovieScene() const
 {
-	return MovieScene;
+	return ActorAnimation ? ActorAnimation->MovieScene : nullptr;
 }
 
-
-UObject* UActorAnimation::GetParentObject(UObject* Object) const
+UObject* UActorAnimationInstance::GetParentObject(UObject* Object) const
 {
 	UActorComponent* Component = Cast<UActorComponent>(Object);
 
@@ -139,11 +207,12 @@ UObject* UActorAnimation::GetParentObject(UObject* Object) const
 	return nullptr;
 }
 
-
-void UActorAnimation::SpawnOrDestroyObjects(bool DestroyAll)
+void UActorAnimationInstance::SpawnOrDestroyObjects(bool DestroyAll)
 {
 	bool bAnyLevelActorsChanged = false;
 
+	UMovieScene* MovieScene = GetMovieScene();
+	
 	// remove any proxy actors that we no longer need
 	if (MovieScene != nullptr)
 	{
@@ -271,9 +340,8 @@ void UActorAnimation::SpawnOrDestroyObjects(bool DestroyAll)
 	}
 }
 
-
 #if WITH_EDITOR
-bool UActorAnimation::TryGetObjectDisplayName(const FGuid& ObjectId, FText& OutDisplayName) const
+bool UActorAnimationInstance::TryGetObjectDisplayName(const FGuid& ObjectId, FText& OutDisplayName) const
 {
 	UObject* Object = FindObject(ObjectId);
 
@@ -295,19 +363,19 @@ bool UActorAnimation::TryGetObjectDisplayName(const FGuid& ObjectId, FText& OutD
 }
 #endif
 
-
-void UActorAnimation::UnbindPossessableObjects(const FGuid& ObjectId)
+void UActorAnimationInstance::UnbindPossessableObjects(const FGuid& ObjectId)
 {
-	// @todo sequencer: gmp: need support for UStruct keys in TMap
-	PossessedObjects.Remove(ObjectId.ToString());
-	//PossessedObjects.Remove(ObjectId);
+	if (bCanRemapBindings && RemappedObjectReferences.HasBinding(ObjectId))
+	{
+		RemappedObjectReferences.RemoveBinding(ObjectId);
+	}
+	else if (ActorAnimation)
+	{
+		ActorAnimation->DefaultObjectReferences.RemoveBinding(ObjectId);
+	}
 }
 
-
-/* UActorAnimation implementation
- *****************************************************************************/
-
-void UActorAnimation::DeselectAllActors()
+void UActorAnimationInstance::DeselectAllActors()
 {
 	// @todo sequencer: gmp: fix me
 	/*
@@ -328,4 +396,4 @@ void UActorAnimation::DeselectAllActors()
 
 		Selection->EndBatchSelectOperation();
 	}*/
-}
+}
\ No newline at end of file
diff --git a/Engine/Source/Runtime/ActorAnimation/Private/ActorAnimationObjectReference.cpp b/Engine/Source/Runtime/ActorAnimation/Private/ActorAnimationObjectReference.cpp
new file mode 100644
--- /dev/null
+++ b/Engine/Source/Runtime/ActorAnimation/Private/ActorAnimationObjectReference.cpp
@@ -0,0 +1,104 @@
+// Copyright 1998-2015 Epic Games, Inc. All Rights Reserved.
+
+#include "ActorAnimationPrivatePCH.h"
+#include "ActorAnimationObjectReference.h"
+
+FActorAnimationObjectReference::FActorAnimationObjectReference(UObject* InObject, UObject* InContext)
+{
+	if (InObject->IsA<AActor>() || InObject->IsA<UActorComponent>())
+	{
+		ObjectId = FLazyObjectPtr(InObject).GetUniqueID();
+		ObjectPath = InObject->GetPathName(InContext);
+	}
+}
+
+UObject* FActorAnimationObjectReference::Resolve(UObject* InContext) const
+{
+	int32 PIEInstanceID = InContext->GetOutermost()->PIEInstanceID;
+	FUniqueObjectGuid FixedUpId = PIEInstanceID == -1 ? ObjectId : ObjectId.FixupForPIE(PIEInstanceID);
+
+	FLazyObjectPtr LazyPtr;
+	LazyPtr = FixedUpId;
+
+	if (UObject* FoundObject = LazyPtr.Get())
+	{
+		return FoundObject;
+	}
+	else if (!ObjectPath.IsEmpty())
+	{
+		return FindObject<UObject>(InContext, *ObjectPath, false);
+	}
+
+	return nullptr;
+}
+
+bool FActorAnimationObjectReferenceMap::Serialize(FArchive& Ar)
+{
+	int32 Num = Map.Num();
+	Ar << Num;
+
+	if (Ar.IsLoading())
+	{
+		while(Num-- > 0)
+		{
+			FGuid Key;
+			Ar << Key;
+
+			FActorAnimationObjectReference Value;
+			Ar << Value;
+
+			Map.Add(Key, Value);
+		}
+	}
+	else if (Ar.IsSaving() || Ar.IsCountingMemory() || Ar.IsObjectReferenceCollector())
+	{
+		for (auto& Pair : Map)
+		{
+			Ar << Pair.Key;
+			Ar << Pair.Value;
+		}
+	}
+	return true;
+}
+
+bool FActorAnimationObjectReferenceMap::HasBinding(const FGuid& ObjectId) const
+{
+	if (const FActorAnimationObjectReference* Reference = Map.Find(ObjectId))
+	{
+		return Reference->IsValid();
+	}
+	return false;
+}
+
+void FActorAnimationObjectReferenceMap::CreateBinding(const FGuid& ObjectId, UObject* InObject, UObject* InContext)
+{
+	FActorAnimationObjectReference NewReference(InObject, InContext);
+
+	if (ensureMsgf(NewReference.IsValid(), TEXT("Unable to generate a reference for the specified object and context")))
+	{
+		Map.FindOrAdd(ObjectId) = NewReference;
+	}
+}
+
+void FActorAnimationObjectReferenceMap::RemoveBinding(const FGuid& ObjectId)
+{
+	Map.Remove(ObjectId);
+}
+
+UObject* FActorAnimationObjectReferenceMap::ResolveBinding(const FGuid& ObjectId, UObject* InContext) const
+{
+	const FActorAnimationObjectReference* Reference = Map.Find(ObjectId);
+	return Reference ? Reference->Resolve(InContext) : nullptr;
+}
+
+FGuid FActorAnimationObjectReferenceMap::FindBindingId(UObject* InObject, UObject* InContext) const
+{
+	for (const auto& Pair : Map)
+	{
+		if (Pair.Value.Resolve(InContext) == InObject)
+		{
+			return Pair.Key;
+		}
+	}
+	return FGuid();
+}
\ No newline at end of file
diff --git a/Engine/Source/Runtime/ActorAnimation/Private/ActorAnimationPlayer.cpp b/Engine/Source/Runtime/ActorAnimation/Private/ActorAnimationPlayer.cpp
--- a/Engine/Source/Runtime/ActorAnimation/Private/ActorAnimationPlayer.cpp
+++ b/Engine/Source/Runtime/ActorAnimation/Private/ActorAnimationPlayer.cpp
@@ -7,21 +7,74 @@
 #include "MovieSceneSequenceInstance.h"
 
 
+struct FTickAnimationPlayers : public FTickableGameObject
+{
+	TArray<TWeakObjectPtr<UActorAnimationPlayer>> ActiveInstances;
+	
+	virtual bool IsTickable() const override
+	{
+		return ActiveInstances.Num() != 0;
+	}
+
+	virtual TStatId GetStatId() const override
+	{
+		RETURN_QUICK_DECLARE_CYCLE_STAT(FTickAnimationPlayers, STATGROUP_Tickables);
+	}
+	
+	virtual void Tick(float DeltaSeconds) override
+	{
+		for (int32 Index = 0; Index < ActiveInstances.Num();)
+		{
+			if (auto* Player = ActiveInstances[Index].Get())
+			{
+				Player->Update(DeltaSeconds);
+				++Index;
+			}
+			else
+			{
+				ActiveInstances.RemoveAt(Index, 1, false);
+			}
+		}
+	}
+};
+
+struct FAutoDestroyAnimationTicker
+{
+	FAutoDestroyAnimationTicker()
+	{
+		FCoreDelegates::OnPreExit.AddLambda([&]{
+			Impl.Reset();
+		});
+	}
+	
+	void Add(UActorAnimationPlayer* Player)
+	{
+		if (!Impl.IsValid())
+		{
+			Impl.Reset(new FTickAnimationPlayers);
+		}
+		Impl->ActiveInstances.Add(Player);
+	}
+
+	TUniquePtr<FTickAnimationPlayers> Impl;
+} GAnimationPlayerTicker;
+
 /* UActorAnimationPlayer structors
  *****************************************************************************/
 
 UActorAnimationPlayer::UActorAnimationPlayer(const FObjectInitializer& ObjectInitializer)
 	: Super(ObjectInitializer)
-	, ActorAnimation(nullptr)
+	, ActorAnimationInstance(nullptr)
 	, bIsPlaying(false)
 	, TimeCursorPosition(0.0f)
+	, CurrentNumLoops(0)
 { }
 
 
 /* UActorAnimationPlayer interface
  *****************************************************************************/
 
-UActorAnimationPlayer* UActorAnimationPlayer::CreateActorAnimationPlayer(UObject* WorldContextObject, UActorAnimation* ActorAnimation)
+UActorAnimationPlayer* UActorAnimationPlayer::CreateActorAnimationPlayer(UObject* WorldContextObject, UActorAnimation* ActorAnimation, const FActorAnimationPlaybackSettings& Settings)
 {
 	if (ActorAnimation == nullptr)
 	{
@@ -29,16 +82,19 @@ UActorAnimationPlayer* UActorAnimationPlayer::CreateActorAnimationPlayer(UObject
 	}
 
 	UWorld* World = GEngine->GetWorldFromContextObject(WorldContextObject);
-
 	check(World != nullptr);
-	check(World->PersistentLevel != nullptr);
-
-	UActorAnimationPlayer* NewPlayer = NewObject<UActorAnimationPlayer>((UObject*)GetTransientPackage(), NAME_None, RF_Transient);
 
+	UActorAnimationPlayer* NewPlayer = NewObject<UActorAnimationPlayer>(GetTransientPackage(), NAME_None, RF_Transient);
 	check(NewPlayer != nullptr);
 
-	NewPlayer->Initialize(ActorAnimation, World);
-	World->PersistentLevel->AddActiveRuntimeMovieScenePlayer(NewPlayer);
+	// Set up a new instance of the actor animation for the player
+	UActorAnimationInstance* Instance = NewObject<UActorAnimationInstance>(NewPlayer);
+	Instance->Initialize(ActorAnimation, World, false);
+
+	NewPlayer->Initialize(Instance, World, Settings);
+
+	// Automatically tick this player
+	GAnimationPlayerTicker.Add(NewPlayer);
 
 	return NewPlayer;
 }
@@ -55,16 +111,24 @@ void UActorAnimationPlayer::Pause()
 	bIsPlaying = false;
 }
 
+void UActorAnimationPlayer::Stop()
+{
+	bIsPlaying = false;
+	TimeCursorPosition = PlaybackSettings.PlaySpeed < 0.f ? GetLength() : 0.f;
+	CurrentNumLoops = 0;
+
+	// todo: Trigger an event?
+}
 
 void UActorAnimationPlayer::Play()
 {
-	if ((ActorAnimation == nullptr) || !World.IsValid())
+	if ((ActorAnimationInstance == nullptr) || !World.IsValid())
 	{
 		return;
 	}
 
 	// @todo Sequencer playback: Should we recreate the instance every time?
-	RootMovieSceneInstance = MakeShareable(new FMovieSceneSequenceInstance(*ActorAnimation));
+	RootMovieSceneInstance = MakeShareable(new FMovieSceneSequenceInstance(*ActorAnimationInstance));
 
 	// @odo Sequencer Should we spawn actors here?
 	SpawnActorsForMovie(RootMovieSceneInstance.ToSharedRef());
@@ -73,14 +137,68 @@ void UActorAnimationPlayer::Play()
 	bIsPlaying = true;
 }
 
+float UActorAnimationPlayer::GetPlaybackPosition() const
+{
+	return TimeCursorPosition;
+}
+
+void UActorAnimationPlayer::SetPlaybackPosition(float NewPlaybackPosition)
+{
+	float LastTimePosition = TimeCursorPosition;
+	
+	TimeCursorPosition = NewPlaybackPosition;
+	OnCursorPositionChanged();
+
+	if (RootMovieSceneInstance.IsValid())
+	{
+		RootMovieSceneInstance->Update(TimeCursorPosition, LastTimePosition, *this);
+	}
+}
+
+float UActorAnimationPlayer::GetLength() const
+{
+	if (!ActorAnimationInstance)
+	{
+		return 0;
+	}
+
+	UMovieScene* MovieScene = ActorAnimationInstance->GetMovieScene();
+	return MovieScene ? MovieScene->GetTimeRange().Size<float>() : 0;
+}
+
+void UActorAnimationPlayer::OnCursorPositionChanged()
+{
+	float Length = GetLength();
+
+	// Handle looping or stopping
+	if (TimeCursorPosition > Length || TimeCursorPosition < 0)
+	{
+		if (PlaybackSettings.LoopCount < 0 || CurrentNumLoops < PlaybackSettings.LoopCount)
+		{
+			++CurrentNumLoops;
+			const float Overplay = FMath::Fmod(TimeCursorPosition, Length);
+			TimeCursorPosition = Overplay < 0 ? Length + Overplay : Overplay;
+		}
+		else
+		{
+			// Stop playing
+			Stop();
+		}
+	}
+}
 
 /* UActorAnimationPlayer implementation
  *****************************************************************************/
 
-void UActorAnimationPlayer::Initialize(UActorAnimation* InActorAnimation, UWorld* InWorld)
+void UActorAnimationPlayer::Initialize(UActorAnimationInstance* InActorAnimationInstance, UWorld* InWorld, const FActorAnimationPlaybackSettings& Settings)
 {
-	ActorAnimation = InActorAnimation;
+	ActorAnimationInstance = InActorAnimationInstance;
+
 	World = InWorld;
+	PlaybackSettings = Settings;
+
+	// Ensure everything is set up, ready for playback
+	Stop();
 }
 
 
@@ -212,7 +330,7 @@ void UActorAnimationPlayer::GetRuntimeObjects(TSharedRef<FMovieSceneSequenceInst
 	else
 	{
 		// otherwise, check whether we have one or more possessed actors that are mapped to this handle
-		UObject* FoundObject = ActorAnimation->FindObject(ObjectId);
+		UObject* FoundObject = ActorAnimationInstance->FindObject(ObjectId);
 
 		if (FoundObject != nullptr)
 		{
@@ -245,23 +363,19 @@ void UActorAnimationPlayer::RemoveMovieSceneInstance(UMovieSceneSection& MovieSc
 	DestroyActorsForMovie( InstanceToRemove );
 }
 
-
 TSharedRef<FMovieSceneSequenceInstance> UActorAnimationPlayer::GetRootMovieSceneSequenceInstance() const
 {
 	return RootMovieSceneInstance.ToSharedRef();
 }
 
-
-/* IRuntimeMovieScenePlayerInterface interface
- *****************************************************************************/
-
-void UActorAnimationPlayer::Tick(const float DeltaSeconds)
+void UActorAnimationPlayer::Update(const float DeltaSeconds)
 {
 	float LastTimePosition = TimeCursorPosition;
 
 	if (bIsPlaying)
 	{
-		TimeCursorPosition += DeltaSeconds;
+		TimeCursorPosition += DeltaSeconds * PlaybackSettings.PlaySpeed;
+		OnCursorPositionChanged();
 	}
 
 	if(RootMovieSceneInstance.IsValid())
diff --git a/Engine/Source/Runtime/ActorAnimation/Private/MovieSceneActor.cpp b/Engine/Source/Runtime/ActorAnimation/Private/MovieSceneActor.cpp
new file mode 100644
--- /dev/null
+++ b/Engine/Source/Runtime/ActorAnimation/Private/MovieSceneActor.cpp
@@ -0,0 +1,99 @@
+// Copyright 1998-2015 Epic Games, Inc. All Rights Reserved.
+
+#include "ActorAnimationPrivatePCH.h"
+#include "MovieSceneActor.h"
+
+
+AMovieSceneActor::AMovieSceneActor(const FObjectInitializer& Init)
+	: Super(Init)
+{
+#if WITH_EDITORONLY_DATA
+	UBillboardComponent* SpriteComponent = CreateEditorOnlyDefaultSubobject<UBillboardComponent>(TEXT("Sprite"));
+
+	if (!IsRunningCommandlet())
+	{
+		// Structure to hold one-time initialization
+		struct FConstructorStatics
+		{
+			ConstructorHelpers::FObjectFinderOptional<UTexture2D> DecalTexture;
+			FConstructorStatics() : DecalTexture(TEXT("/Engine/EditorResources/SceneManager")) {}
+		};
+		static FConstructorStatics ConstructorStatics;
+
+		if (SpriteComponent)
+		{
+			SpriteComponent->Sprite = ConstructorStatics.DecalTexture.Get();
+			SpriteComponent->AttachParent = RootComponent;
+			SpriteComponent->bIsScreenSizeScaled = true;
+			SpriteComponent->bAbsoluteScale = true;
+			SpriteComponent->bReceivesDecals = false;
+			SpriteComponent->bHiddenInGame = true;
+			SetRootComponent(SpriteComponent);
+		}
+	}
+#endif // WITH_EDITORONLY_DATA
+
+	PrimaryActorTick.bCanEverTick = true;
+	bAutoPlay = false;
+}
+
+void AMovieSceneActor::BeginPlay()
+{
+	Super::BeginPlay();
+	InitializePlayer();
+}
+
+void AMovieSceneActor::PostInitProperties()
+{
+	Super::PostInitProperties();
+
+	if (!HasAnyFlags(RF_ClassDefaultObject))
+	{
+		UpdateAnimationInstance();
+	}
+}
+
+void AMovieSceneActor::Tick(float DeltaSeconds)
+{
+	if (AnimationPlayer)
+	{
+		AnimationPlayer->Update(DeltaSeconds);
+	}
+}
+
+void AMovieSceneActor::SetAnimation(UActorAnimation* InAnimation)
+{
+	if (!AnimationPlayer || !AnimationPlayer->IsPlaying())
+	{
+		ActorAnimation = InAnimation;
+		InitializePlayer();
+	}
+}
+
+void AMovieSceneActor::InitializePlayer()
+{
+	UpdateAnimationInstance();
+
+	if (GetWorld()->IsGameWorld())
+	{
+		AnimationPlayer = NewObject<UActorAnimationPlayer>(this, "AnimationPlayer");
+		AnimationPlayer->Initialize(AnimationInstance, GetWorld(), PlaybackSettings);
+
+		if (bAutoPlay)
+		{
+			AnimationPlayer->Play();
+		}
+	}
+}
+
+void AMovieSceneActor::UpdateAnimationInstance()
+{
+	UActorAnimation* ActorAnimationAsset = Cast<UActorAnimation>(ActorAnimation.TryLoad());
+	if (!AnimationInstance)
+	{
+		// Make an instance of our asset so that we can keep hard references to actors we are using
+		AnimationInstance = NewObject<UActorAnimationInstance>(this, "AnimationInstance");
+	}
+
+	AnimationInstance->Initialize(ActorAnimationAsset, GetWorld(), true);
+}
\ No newline at end of file
diff --git a/Engine/Source/Runtime/ActorAnimation/Public/ActorAnimation.h b/Engine/Source/Runtime/ActorAnimation/Public/ActorAnimation.h
--- a/Engine/Source/Runtime/ActorAnimation/Public/ActorAnimation.h
+++ b/Engine/Source/Runtime/ActorAnimation/Public/ActorAnimation.h
@@ -4,25 +4,73 @@
 
 #include "MovieSceneSequence.h"
 #include "ActorAnimationObject.h"
+#include "ActorAnimationObjectReference.h"
 #include "ActorAnimation.generated.h"
 
-
 /**
  * Movie scene animation for Actors.
  */
-UCLASS(BlueprintType, MinimalAPI)
-class UActorAnimation
-	: public UMovieSceneSequence
+UCLASS(BlueprintType)
+class ACTORANIMATION_API UActorAnimation
+	: public UObject
 {
 	GENERATED_UCLASS_BODY()
 
+public:
+
+	/** Collection of possessed objects. */
+	UPROPERTY()
+	FActorAnimationObjectReferenceMap DefaultObjectReferences;
+
+	/** Pointer to the movie scene that controls this animation. */
+	UPROPERTY()
+	UMovieScene* MovieScene;
+
+public:
+
 	/** Initialize this actor animation. */
-	ACTORANIMATION_API void Initialize();
+	void Initialize();
+
+	/** Convert old-style lazy object ptrs to new-style references using the specified context */
+	void ConvertPersistentBindingsToDefault(UObject* FixupContext);
+
+private:
 
+	/** Deprecated property housing old possessed object bindings */
+	UPROPERTY()
+	TMap<FString, FActorAnimationObject> PossessedObjects_DEPRECATED;
+};
+
+/**
+ * An instance of a UActorAnimation that supports resolution and remapping of object bindings
+ */
+UCLASS()
+class ACTORANIMATION_API UActorAnimationInstance
+	: public UMovieSceneSequence
+{
 public:
 
-	// UMovieSceneAnimation overrides
+	GENERATED_BODY()
+	
+	/**
+	 * Initialize this instance. Safe to be called multiple times. Does not reset object remappings unless bInCanRemapBindings is false.
+	 *
+	 * @param InAnimation			The actor animation to instatiate
+	 * @param InContext				The context to use for object binding resolution
+	 * @param bInCanRemapBindings	Whether this instance is permitted to remap object bindings or not
+	 */
+	void Initialize(UActorAnimation* InAnimation, UObject* InContext, bool bInCanRemapBindings);
+	
+	/**
+	 * Change this instance's context
+	 *
+	 * @param NewContext			The new context to use for object binding resolution
+	 */
+	void SetContext(UObject* NewContext);
+
+public:
 
+	// UMovieSceneSequence interface
 	virtual bool AllowsSpawnableObjects() const override;
 	virtual void BindPossessableObject(const FGuid& ObjectId, UObject& PossessedObject) override;
 	virtual bool CanPossessObject(UObject& Object) const override;
@@ -33,11 +81,11 @@ public:
 	virtual UObject* GetParentObject(UObject* Object) const override;
 	virtual void SpawnOrDestroyObjects(bool DestroyAll) override;
 	virtual void UnbindPossessableObjects(const FGuid& ObjectId) override;
-
 #if WITH_EDITOR
 	virtual bool TryGetObjectDisplayName(const FGuid& ObjectId, FText& OutDisplayName) const override;
 #endif
 
+	// End UMovieSceneSequence interface
 protected:
 
 	/**
@@ -45,20 +93,30 @@ protected:
 	 *
 	 * @todo sequencer: remove Editor dependencies from this class
 	 */
-	ACTORANIMATION_API void DeselectAllActors();
+	void DeselectAllActors();
 
 private:
 
-	/** Pointer to the movie scene that controls this animation. */
+	/** The actor animation we are instancing */
+	UPROPERTY(transient)
+	UActorAnimation* ActorAnimation;
+
+	/** Flag to specify whether this instance can remap bindings or not */
 	UPROPERTY()
-	UMovieScene* MovieScene;
+	bool bCanRemapBindings;
 
-	/** Collection of possessed objects. */
-	// @todo sequencer: gmp: need support for UStruct keys in TMap
+	/** Collection of remappings for possessed objects references. */
 	UPROPERTY()
-	TMap<FString, FActorAnimationObject> PossessedObjects;
-	//TMap<FGuid, FSequencerBinding> Bindings;
+	FActorAnimationObjectReferenceMap RemappedObjectReferences;
+
+private:
+
+	/** A context to use for resolving object bindings */
+	TWeakObjectPtr<UObject> ResolutionContext;
+
+	/** A transient map of cached object bindings */
+	mutable TMap<FGuid, TWeakObjectPtr<UObject>> CachedObjectBindings;
 
 	/** Collection of spawned proxy actors. */
 	TMap<FGuid, TWeakObjectPtr<AActor>> SpawnedActors;
-};
+};
\ No newline at end of file
diff --git a/Engine/Source/Runtime/ActorAnimation/Public/ActorAnimationObject.h b/Engine/Source/Runtime/ActorAnimation/Public/ActorAnimationObject.h
--- a/Engine/Source/Runtime/ActorAnimation/Public/ActorAnimationObject.h
+++ b/Engine/Source/Runtime/ActorAnimation/Public/ActorAnimationObject.h
@@ -11,7 +11,7 @@
 USTRUCT()
 struct FActorAnimationObject
 {
-	GENERATED_USTRUCT_BODY(FActorAnimationObject)
+	GENERATED_BODY()
 
 public:
 
diff --git a/Engine/Source/Runtime/ActorAnimation/Public/ActorAnimationObjectReference.h b/Engine/Source/Runtime/ActorAnimation/Public/ActorAnimationObjectReference.h
new file mode 100644
--- /dev/null
+++ b/Engine/Source/Runtime/ActorAnimation/Public/ActorAnimationObjectReference.h
@@ -0,0 +1,151 @@
+// Copyright 1998-2015 Epic Games, Inc. All Rights Reserved.
+
+#pragma once
+
+#include "ActorAnimationObjectReference.generated.h"
+
+/**
+ * An external reference to an actor animation object, resolvable through an arbitrary context
+ */
+USTRUCT()
+struct FActorAnimationObjectReference
+{
+	/**
+	 * Default construction to a null reference
+	 */
+	FActorAnimationObjectReference() : ObjectPath() {}
+
+	/**
+	 * Generates a new reference to an object within a given context.
+	 *
+	 * @param InObject The object to generate a reference for (e.g. an AActor)
+	 * @param InContext The object to use as a context for the generation (e.g. a UWorld)
+	 */
+	ACTORANIMATION_API FActorAnimationObjectReference(UObject* InObject, UObject* InContext);
+
+	GENERATED_BODY();
+
+	/**
+	 * Check whether this object reference is valid or not
+	 */
+	bool IsValid() const
+	{
+		return !ObjectPath.IsEmpty();
+	}
+
+	/**
+	 * Resolve this reference within the specified context
+	 *
+	 * @return The object (usually an Actor or an ActorComponent).
+	 */
+	ACTORANIMATION_API UObject* Resolve(UObject* InContext) const;
+
+	/**
+	 * Equality comparator
+	 */
+	friend bool operator==(const FActorAnimationObjectReference& A, const FActorAnimationObjectReference& B)
+	{
+		if (A.ObjectId.IsValid() && A.ObjectId == B.ObjectId)
+		{
+			return true;
+		}
+		return A.ObjectPath == B.ObjectPath;
+	}
+
+	/**
+	 * Serialization operator
+	 */
+	friend FArchive& operator<<(FArchive& Ar, FActorAnimationObjectReference& Ref)
+	{
+		Ar << Ref.ObjectId;
+		Ar << Ref.ObjectPath;
+		return Ar;
+	}
+
+private:
+
+	/** Primary method of resolution - object ID, stored as an annotation on the object in the world, resolvable through TLazyObjectPtr */
+	FUniqueObjectGuid ObjectId;
+
+	/** Secondary method of resolution - path to the object within the context */
+	FString ObjectPath;
+};
+
+USTRUCT()
+struct FActorAnimationObjectReferenceMap
+{
+	/**
+	 * Check whether this map has a binding for the specified object id
+	 * @return true if this map contains a binding for the id, false otherwise
+	 */
+	bool HasBinding(const FGuid& ObjectId) const;
+
+	/**
+	 * Remove a binding for the specified ID
+	 *
+	 * @param ObjectId	The ID to remove
+	 */
+	void RemoveBinding(const FGuid& ObjectId);
+
+	/**
+	 * Create a binding for the specified ID
+	 *
+	 * @param ObjectId	The ID to associate the object with
+	 * @param InObject	The object to associate
+	 * @param InContext	A context in which InObject resides
+	 */
+	void CreateBinding(const FGuid& ObjectId, UObject* InObject, UObject* InContext);
+
+	/**
+	 * Resolve a binding for the specified ID using a given context
+	 *
+	 * @param ObjectId	The ID to associate the object with
+	 * @param InContext	A context in which InObject resides
+	 *
+	 * @return The object if found, nullptr otherwise.
+	 */
+	UObject* ResolveBinding(const FGuid& ObjectId, UObject* InContext) const;
+
+	/**
+	 * Find an object ID that the specified object is bound with
+	 *
+	 * @param InObject	The object to find an ID for
+	 * @param InContext	The context in which InObject resides
+	 *
+	 * @return The object's bound ID, or an empty FGuid
+	 */
+	FGuid FindBindingId(UObject* InObject, UObject* InContext) const;
+
+public:
+
+	GENERATED_BODY();
+
+	/**
+	 * Equality comparator required for proper UObject serialization (so we can compare against defaults)
+	 */
+	friend bool operator==(const FActorAnimationObjectReferenceMap& A, const FActorAnimationObjectReferenceMap& B)
+	{
+		return A.Map.OrderIndependentCompareEqual(B.Map);
+	}
+
+	/**
+	 * Serialization function
+	 */
+	bool Serialize(FArchive& Ar);
+
+	typedef TMap<FGuid, FActorAnimationObjectReference> MapType;
+	FORCEINLINE friend MapType::TIterator      begin(      FActorAnimationObjectReferenceMap& Impl) { return begin(Impl.Map); }
+	FORCEINLINE friend MapType::TConstIterator begin(const FActorAnimationObjectReferenceMap& Impl) { return begin(Impl.Map); }
+	FORCEINLINE friend MapType::TIterator      end  (      FActorAnimationObjectReferenceMap& Impl) { return end(Impl.Map); }
+	FORCEINLINE friend MapType::TConstIterator end  (const FActorAnimationObjectReferenceMap& Impl) { return end(Impl.Map); }
+
+private:
+	
+	// @todo sequencer: need support for UStruct keys in TMap so this can be a UPROPERTY
+	TMap<FGuid, FActorAnimationObjectReference> Map;
+};
+
+template<> struct TStructOpsTypeTraits<FActorAnimationObjectReferenceMap> : public TStructOpsTypeTraitsBase
+{
+	enum { WithSerializer = true, WithIdenticalViaEquality = true };
+};
\ No newline at end of file
diff --git a/Engine/Source/Runtime/ActorAnimation/Public/ActorAnimationPlayer.h b/Engine/Source/Runtime/ActorAnimation/Public/ActorAnimationPlayer.h
--- a/Engine/Source/Runtime/ActorAnimation/Public/ActorAnimationPlayer.h
+++ b/Engine/Source/Runtime/ActorAnimation/Public/ActorAnimationPlayer.h
@@ -4,7 +4,6 @@
 
 #include "IMovieScenePlayer.h"
 #include "ActorAnimation.h"
-#include "MovieScene/RuntimeMovieScenePlayerInterface.h"
 #include "ActorAnimationPlayer.generated.h"
 
 
@@ -13,6 +12,25 @@ class ULevel;
 class UMovieSceneBindings;
 
 
+USTRUCT(BlueprintType)
+struct FActorAnimationPlaybackSettings
+{
+	FActorAnimationPlaybackSettings()
+		: LoopCount(-1)
+		, PlaySpeed(1.f)
+	{}
+
+	GENERATED_BODY()
+
+	/** Number of times to loop playback. -1 for infinite, else the number of times to loop before stopping */
+	UPROPERTY(EditAnywhere, BlueprintReadWrite, Category="Playback", meta=(UIMin=-1))
+	int32 LoopCount;
+
+	/** The speed at which to playback the animation */
+	UPROPERTY(EditAnywhere, BlueprintReadWrite, Category="Playback")
+	float PlaySpeed;
+};
+
 /**
  * UActorAnimationPlayer is used to actually "play" an actor animation asset at runtime.
  *
@@ -23,9 +41,20 @@ UCLASS(BlueprintType)
 class ACTORANIMATION_API UActorAnimationPlayer
 	: public UObject
 	, public IMovieScenePlayer
-	, public IRuntimeMovieScenePlayerInterface
 {
-	GENERATED_UCLASS_BODY()
+public:
+	UActorAnimationPlayer(const FObjectInitializer&);
+
+	GENERATED_BODY()
+
+	/**
+	 * Initialize the player.
+	 *
+	 * @param InActorAnimationInstance The actor animation instance to play.
+	 * @param InWorld The world that the animation is played in.
+	 * @param Settings The desired playback settings
+	 */
+	void Initialize(UActorAnimationInstance* InActorAnimationInstance, UWorld* InWorld, const FActorAnimationPlaybackSettings& Settings);
 
 public:
 
@@ -34,9 +63,10 @@ public:
 	 *
 	 * @param WorldContextObject Context object from which to retrieve a UWorld.
 	 * @param ActorAnimation The actor animation to play.
+	 * @param Settings The desired playback settings
 	 */
 	UFUNCTION(BlueprintCallable, Category="Game|Cinematic", meta=(WorldContext="WorldContextObject"))
-	static UActorAnimationPlayer* CreateActorAnimationPlayer(UObject* WorldContextObject, UActorAnimation* ActorAnimation);
+	static UActorAnimationPlayer* CreateActorAnimationPlayer(UObject* WorldContextObject, UActorAnimation* ActorAnimation, const FActorAnimationPlaybackSettings& Settings);
 
 	/** Start playback from the current time cursor position. */
 	UFUNCTION(BlueprintCallable, Category="Game|Cinematic")
@@ -45,19 +75,26 @@ public:
 	/** Pause playback. */
 	UFUNCTION(BlueprintCallable, Category="Game|Cinematic")
 	void Pause();
+	
+	/** Stop playback. */
+	UFUNCTION(BlueprintCallable, Category="Game|Cinematic")
+	void Stop();
+
+	/** Get the current playback position */
+	UFUNCTION(BlueprintCallable, Category="Game|Cinematic")
+	float GetPlaybackPosition() const;
+
+	/** Set the current playback position */
+	UFUNCTION(BlueprintCallable, Category="Game|Cinematic")
+	void SetPlaybackPosition(float NewPlaybackPosition);
 
 	/** Check whether the sequence is actively playing. */
+	UFUNCTION(BlueprintCallable, Category="Game|Cinematic")
 	bool IsPlaying() const;
 
-protected:
-
-	/**
-	 * Initialize the player.
-	 *
-	 * @param InActorAnimation The actor animation to play.
-	 * @param InWorld The world that the animation is played in.
-	 */
-	void Initialize(UActorAnimation* InActorAnimation, UWorld* InWorld);
+	/** Get the length of the sequence */
+	UFUNCTION(BlueprintCallable, Category="Game|Cinematic")
+	float GetLength() const;
 
 protected:
 
@@ -72,17 +109,20 @@ protected:
 	virtual void RemoveMovieSceneInstance(UMovieSceneSection& MovieSceneSection, TSharedRef<FMovieSceneSequenceInstance> InstanceToRemove) override;
 	virtual TSharedRef<FMovieSceneSequenceInstance> GetRootMovieSceneSequenceInstance() const override;
 
-protected:
+public:
+
+	void Update(const float DeltaSeconds);
 
-	// IRuntimeMovieScenePlayerInterface
+private:
 
-	virtual void Tick(const float DeltaSeconds) override;
+	/** Called when the cursor position has changed to implement looping */
+	void OnCursorPositionChanged();
 
 private:
 
 	/** The actor animation to play. */
 	UPROPERTY()
-	UActorAnimation* ActorAnimation;
+	UActorAnimationInstance* ActorAnimationInstance;
 
 	/** Whether we're currently playing. If false, then sequence playback is paused or was never started. */
 	UPROPERTY()
@@ -90,7 +130,14 @@ private:
 
 	/** The current time cursor position within the sequence (in seconds) */
 	UPROPERTY()
-	double TimeCursorPosition;
+	float TimeCursorPosition;
+
+	/** Specific playback settings for the animation. */
+	UPROPERTY()
+	FActorAnimationPlaybackSettings PlaybackSettings;
+
+	/** The number of times we have looped in the current playback */
+	int32 CurrentNumLoops;
 
 private:
 
diff --git a/Engine/Source/Runtime/ActorAnimation/Public/MovieSceneActor.h b/Engine/Source/Runtime/ActorAnimation/Public/MovieSceneActor.h
new file mode 100644
--- /dev/null
+++ b/Engine/Source/Runtime/ActorAnimation/Public/MovieSceneActor.h
@@ -0,0 +1,52 @@
+// Copyright 1998-2015 Epic Games, Inc. All Rights Reserved.
+
+#pragma once
+
+#include "ActorAnimationPlayer.h"
+#include "MovieSceneActor.generated.h"
+
+/**
+ * Actor responsible for controlling a specific movie scene sequence in the world
+ */
+UCLASS(Experimental, hideCategories=(Rendering, Physics, LOD, Activation))
+class ACTORANIMATION_API AMovieSceneActor : public AActor
+{
+public:
+	AMovieSceneActor(const FObjectInitializer& Init);
+
+	GENERATED_BODY()
+
+	UPROPERTY(EditAnywhere, BlueprintReadWrite, Category="Playback")
+	bool bAutoPlay;
+
+	UPROPERTY(EditAnywhere, BlueprintReadWrite, Category="Playback", meta=(ShowOnlyInnerProperties))
+	FActorAnimationPlaybackSettings PlaybackSettings;
+
+	UPROPERTY(transient, BlueprintReadOnly, Category="Playback")
+	UActorAnimationPlayer* AnimationPlayer;
+
+	UPROPERTY(EditAnywhere, BlueprintReadWrite, Category="General", meta=(AllowedClasses="ActorAnimation"))
+	FStringAssetReference ActorAnimation;
+
+public:
+
+	UFUNCTION(BlueprintCallable, Category="Game|Cinematic")
+	void SetAnimation(UActorAnimation* Animation);
+
+public:
+
+	virtual void BeginPlay() override;
+	virtual void Tick(float DeltaSeconds) override;
+	virtual void PostInitProperties() override;
+
+	void InitializePlayer();
+
+private:
+
+	/** Instance of the actor animation we're using. Stored as a property so we can track live actor references */
+	UPROPERTY(Instanced)
+	UActorAnimationInstance* AnimationInstance;
+
+	/** Update the instance of our actor animation */
+	void UpdateAnimationInstance();
+};
diff --git a/Engine/Source/Runtime/CoreUObject/Private/UObject/LazyObjectPtr.cpp b/Engine/Source/Runtime/CoreUObject/Private/UObject/LazyObjectPtr.cpp
--- a/Engine/Source/Runtime/CoreUObject/Private/UObject/LazyObjectPtr.cpp
+++ b/Engine/Source/Runtime/CoreUObject/Private/UObject/LazyObjectPtr.cpp
@@ -23,12 +23,12 @@ FUniqueObjectGuid::FUniqueObjectGuid(const class UObject* InObject)
 {
 }
 
-FUniqueObjectGuid FUniqueObjectGuid::FixupForPIE() const
+FUniqueObjectGuid FUniqueObjectGuid::FixupForPIE(int32 PlayInEditorID) const
 {
 	FUniqueObjectGuid Temp(*this);
 
-	check(GPlayInEditorID != -1)
-	const FGuid *FoundGuid = PIEGuidMap[GPlayInEditorID % MAX_PIE_INSTANCES].Find(Temp.GetGuid());
+	check(PlayInEditorID != -1)
+	const FGuid *FoundGuid = PIEGuidMap[PlayInEditorID % MAX_PIE_INSTANCES].Find(Temp.GetGuid());
 
 	if (FoundGuid)
 	{
diff --git a/Engine/Source/Runtime/CoreUObject/Public/UObject/LazyObjectPtr.h b/Engine/Source/Runtime/CoreUObject/Public/UObject/LazyObjectPtr.h
--- a/Engine/Source/Runtime/CoreUObject/Public/UObject/LazyObjectPtr.h
+++ b/Engine/Source/Runtime/CoreUObject/Public/UObject/LazyObjectPtr.h
@@ -46,7 +46,7 @@ struct COREUOBJECT_API FUniqueObjectGuid
 	/**
 	 * Fixes up this UniqueObjectID to add or remove the PIE prefix depending on what is currently active
 	 */
-	FUniqueObjectGuid FixupForPIE() const;
+	FUniqueObjectGuid FixupForPIE(int32 PlayInEditorID = GPlayInEditorID) const;
 
 	/**
 	 * Attempts to find a currently loaded object that matches this object ID
diff --git a/Engine/Source/Runtime/Engine/Classes/Engine/Level.h b/Engine/Source/Runtime/Engine/Classes/Engine/Level.h
--- a/Engine/Source/Runtime/Engine/Classes/Engine/Level.h
+++ b/Engine/Source/Runtime/Engine/Classes/Engine/Level.h
@@ -496,16 +496,6 @@ private:
 	FLevelBoundsActorUpdatedEvent LevelBoundsActorUpdatedEvent; 
 
 protected:
-	/** Array of all MovieSceneBindings that are used in this level.  These store the relationship between
-	    a MovieScene asset and possessed actors in this level. */
-	UPROPERTY()
-	TArray< class UObject* > MovieSceneBindingsArray;
-
-	/** List of RuntimeMovieScenePlayers that are currently active in this level.  We'll keep references to these to keep
-	    them around until they're no longer needed.  Also, we'll tick them every frame! */
-	// @todo sequencer uobjects: Ideally this is using URuntimeMovieScenePlayer* and not UObject*, but there are DLL/interface issues with that
-	UPROPERTY( transient )
-	TArray< UObject* > ActiveRuntimeMovieScenePlayers;
 
 	/** Array of user data stored with the asset */
 	UPROPERTY()
@@ -553,20 +543,6 @@ public:
 	static void AddReferencedObjects(UObject* InThis, FReferenceCollector& Collector);
 	//~ End UObject Interface.
 
-	/**
-	 * Adds a newly-created RuntimeMovieScenePlayer to this level.  The level assumes ownership of this object.
-	 *
-	 * @param	RuntimeMovieScenePlayer	The RuntimeMovieScenePlayer instance to add
-	 */
-	ENGINE_API void AddActiveRuntimeMovieScenePlayer( class UObject* RuntimeMovieScenePlayer );
-
-	/**
-	 * Called by the engine every frame to tick all actively-playing RuntimeMovieScenePlayers
-	 *
-	 * @param	DeltaSeconds	Time elapsed since last tick
-	 */
-	void TickRuntimeMovieScenePlayers( const float DeltaSeconds );
-
 	/**
 	 * Clears all components of actors associated with this level (aka in Actors array) and 
 	 * also the BSP model components.
@@ -770,21 +746,6 @@ public:
 	/** meant to be called only from editor, calculating and storing static geometry to be used with off-line and/or on-line navigation building */
 	ENGINE_API void RebuildStaticNavigableGeometry();
 
-
-	/**
-	 * Adds a new MovieSceneBindings object to the level script actor.  This actor takes full ownership of
-	 * the MovieSceneBindings
-	 *
-	 * @param	MovieSceneBindings	The MovieSceneBindings object to take ownership of
-	 */
-	virtual void AddMovieSceneBindings( class UMovieSceneBindings* MovieSceneBindings );
-
-	/**
-	 * Clears all existing MovieSceneBindings off of the level.  Only meant to be called by the level script compiler.
-	 */
-	virtual void ClearMovieSceneBindings();
-
-
 #endif
 
 };
diff --git a/Engine/Source/Runtime/Engine/Classes/MovieScene/RuntimeMovieScenePlayerInterface.h b/Engine/Source/Runtime/Engine/Classes/MovieScene/RuntimeMovieScenePlayerInterface.h
deleted file mode 100644
diff --git a/Engine/Source/Runtime/Engine/Private/Interface.cpp b/Engine/Source/Runtime/Engine/Private/Interface.cpp
--- a/Engine/Source/Runtime/Engine/Private/Interface.cpp
+++ b/Engine/Source/Runtime/Engine/Private/Interface.cpp
@@ -3,19 +3,12 @@
 #include "EnginePrivate.h"
 #include "Interfaces/Interface_CollisionDataProvider.h"
 #include "Interfaces/Interface_PostProcessVolume.h"
-#include "MovieScene/RuntimeMovieScenePlayerInterface.h"
 
 UInterface_CollisionDataProvider::UInterface_CollisionDataProvider(const FObjectInitializer& ObjectInitializer)
 	: Super(ObjectInitializer)
 {
 }
 
-URuntimeMovieScenePlayerInterface::URuntimeMovieScenePlayerInterface( const FObjectInitializer& ObjectInitializer )
-	: Super(ObjectInitializer)
-{
-
-}
-
 UInterface_PostProcessVolume::UInterface_PostProcessVolume( const FObjectInitializer& ObjectInitializer )
 	: Super(ObjectInitializer)
 {
diff --git a/Engine/Source/Runtime/Engine/Private/Level.cpp b/Engine/Source/Runtime/Engine/Private/Level.cpp
--- a/Engine/Source/Runtime/Engine/Private/Level.cpp
+++ b/Engine/Source/Runtime/Engine/Private/Level.cpp
@@ -22,7 +22,6 @@ Level.cpp: Level-related functions
 #include "Editor/UnrealEd/Public/Kismet2/KismetEditorUtilities.h"
 #include "Editor/UnrealEd/Public/Kismet2/BlueprintEditorUtils.h"
 #endif
-#include "Runtime/Engine/Classes/MovieScene/RuntimeMovieScenePlayerInterface.h"
 #include "LevelUtils.h"
 #include "TargetPlatform.h"
 #include "ContentStreaming.h"
@@ -1826,53 +1825,6 @@ void ULevel::OnLevelScriptBlueprintChanged(ULevelScriptBlueprint* InBlueprint)
 
 #endif	//WITH_EDITOR
 
-
-#if WITH_EDITOR
-void ULevel::AddMovieSceneBindings( class UMovieSceneBindings* MovieSceneBindings )
-{
-	// @todo sequencer major: We need to clean up stale bindings objects that no PlayMovieScene node references anymore (LevelScriptBlueprint compile time?)
-	if( ensure( MovieSceneBindings != NULL ) )
-	{
-		Modify();
-
-		// @todo sequencer UObjects: Dangerous cast here to work around MovieScene not being a dependency module of engine
-		UObject* ObjectToAdd = (UObject*)MovieSceneBindings;
-		MovieSceneBindingsArray.AddUnique( ObjectToAdd );
-	}
-}
-
-void ULevel::ClearMovieSceneBindings()
-{
-	Modify();
-
-	MovieSceneBindingsArray.Reset();
-}
-#endif
-
-void ULevel::AddActiveRuntimeMovieScenePlayer( UObject* RuntimeMovieScenePlayer )
-{
-	ActiveRuntimeMovieScenePlayers.Add( RuntimeMovieScenePlayer );
-}
-
-
-void ULevel::TickRuntimeMovieScenePlayers( const float DeltaSeconds )
-{
-	for( int32 CurPlayerIndex = 0; CurPlayerIndex < ActiveRuntimeMovieScenePlayers.Num(); ++CurPlayerIndex )
-	{
-		// Should never have any active RuntimeMovieScenePlayers on an editor world!
-		check( OwningWorld->IsGameWorld() );
-
-		IRuntimeMovieScenePlayerInterface* RuntimeMovieScenePlayer =
-			Cast< IRuntimeMovieScenePlayerInterface >( ActiveRuntimeMovieScenePlayers[ CurPlayerIndex ] );
-		check( RuntimeMovieScenePlayer != NULL );
-
-		// @todo sequencer runtime: Support expiring instances of RuntimeMovieScenePlayers that have finished playing
-		// @todo sequencer runtime: Support destroying spawned actors for a completed moviescene
-		RuntimeMovieScenePlayer->Tick( DeltaSeconds );
-	}
-}
-
-
 bool ULevel::IsPersistentLevel() const
 {
 	bool bIsPersistent = false;
diff --git a/Engine/Source/Runtime/Engine/Private/LevelTick.cpp b/Engine/Source/Runtime/Engine/Private/LevelTick.cpp
--- a/Engine/Source/Runtime/Engine/Private/LevelTick.cpp
+++ b/Engine/Source/Runtime/Engine/Private/LevelTick.cpp
@@ -1239,21 +1239,7 @@ void UWorld::Tick( ELevelTick TickType, float DeltaSeconds )
 			}
 		}
 	}
-
-	// Tick any RuntimeMovieScenePlayers that our levels know about
-	{
-		// Don't bother ticking RuntimeMovieScenePlayers when paused
-		if( TickType != LEVELTICK_TimeOnly && !bIsPaused )
-		{
-			SCOPE_CYCLE_COUNTER(STAT_RuntimeMovieSceneTickTime);
-			for( int32 LevelIndex = 0; LevelIndex < Levels.Num(); ++LevelIndex )
-			{
-				ULevel* Level = Levels[ LevelIndex ];
-				Level->TickRuntimeMovieScenePlayers( DeltaSeconds );
-			}
-		}
-	}
-
+	
 	// Update cameras and streaming volumes
 	{
 		SCOPE_CYCLE_COUNTER(STAT_UpdateCameraTime);
diff --git a/Engine/Source/Runtime/Engine/Public/Engine.h b/Engine/Source/Runtime/Engine/Public/Engine.h
--- a/Engine/Source/Runtime/Engine/Public/Engine.h
+++ b/Engine/Source/Runtime/Engine/Public/Engine.h
@@ -267,7 +267,6 @@
 #include "Engine/ChildConnection.h"
 #include "Engine/Polys.h"
 #include "Sound/ReverbEffect.h"
-#include "MovieScene/RuntimeMovieScenePlayerInterface.h"
 #include "GameFramework/SaveGame.h"
 #include "Engine/SCS_Node.h"
 #include "Engine/Selection.h"
