Improved Actor Animation Player functionality.
Properties on MovieSceneActor are now BlueprintReadOnly since they are only really set up parameters.
When actor animation players reach the end of their animation, they no longer reset back to the beginning.
Added new PlayLooping function that allows the user to dynamically set how many loops to play.
Changed defaults on playback settings to play a single animation as this is probably a better default.

diff --git a/Engine/Source/Runtime/ActorAnimation/Private/ActorAnimationPlayer.cpp b/Engine/Source/Runtime/ActorAnimation/Private/ActorAnimationPlayer.cpp
--- a/Engine/Source/Runtime/ActorAnimation/Private/ActorAnimationPlayer.cpp
+++ b/Engine/Source/Runtime/ActorAnimation/Private/ActorAnimationPlayer.cpp
@@ -74,7 +74,7 @@ UActorAnimationPlayer::UActorAnimationPlayer(const FObjectInitializer& ObjectIni
 /* UActorAnimationPlayer interface
  *****************************************************************************/
 
-UActorAnimationPlayer* UActorAnimationPlayer::CreateActorAnimationPlayer(UObject* WorldContextObject, UActorAnimation* ActorAnimation, const FActorAnimationPlaybackSettings& Settings)
+UActorAnimationPlayer* UActorAnimationPlayer::CreateActorAnimationPlayer(UObject* WorldContextObject, UActorAnimation* ActorAnimation, FActorAnimationPlaybackSettings Settings)
 {
 	if (ActorAnimation == nullptr)
 	{
@@ -114,7 +114,7 @@ void UActorAnimationPlayer::Pause()
 void UActorAnimationPlayer::Stop()
 {
 	bIsPlaying = false;
-	TimeCursorPosition = PlaybackSettings.PlaySpeed < 0.f ? GetLength() : 0.f;
+	TimeCursorPosition = PlaybackSettings.PlayRate < 0.f ? GetLength() : 0.f;
 	CurrentNumLoops = 0;
 
 	// todo: Trigger an event?
@@ -133,8 +133,19 @@ void UActorAnimationPlayer::Play()
 	// @odo Sequencer Should we spawn actors here?
 	SpawnActorsForMovie(RootMovieSceneInstance.ToSharedRef());
 	RootMovieSceneInstance->RefreshInstance(*this);
-	
+
 	bIsPlaying = true;
+
+	if (RootMovieSceneInstance.IsValid())
+	{
+		RootMovieSceneInstance->Update(TimeCursorPosition, TimeCursorPosition, *this);
+	}
+}
+
+void UActorAnimationPlayer::PlayLooping(int32 NumLoops)
+{
+	PlaybackSettings.LoopCount = NumLoops;
+	Play();
 }
 
 float UActorAnimationPlayer::GetPlaybackPosition() const
@@ -166,6 +177,17 @@ float UActorAnimationPlayer::GetLength() const
 	return MovieScene ? MovieScene->GetTimeRange().Size<float>() : 0;
 }
 
+float UActorAnimationPlayer::GetPlayRate() const
+{
+	return PlaybackSettings.PlayRate;
+}
+
+void UActorAnimationPlayer::SetPlayRate(float PlayRate)
+{
+	PlaybackSettings.PlayRate = PlayRate;
+}
+
+
 void UActorAnimationPlayer::OnCursorPositionChanged()
 {
 	float Length = GetLength();
@@ -181,8 +203,10 @@ void UActorAnimationPlayer::OnCursorPositionChanged()
 		}
 		else
 		{
-			// Stop playing
-			Stop();
+			// Stop playing without modifying the playback position
+			// @todo: trigger an event?
+			bIsPlaying = false;
+			CurrentNumLoops = 0;
 		}
 	}
 }
@@ -374,7 +398,7 @@ void UActorAnimationPlayer::Update(const float DeltaSeconds)
 
 	if (bIsPlaying)
 	{
-		TimeCursorPosition += DeltaSeconds * PlaybackSettings.PlaySpeed;
+		TimeCursorPosition += DeltaSeconds * PlaybackSettings.PlayRate;
 		OnCursorPositionChanged();
 	}
 
diff --git a/Engine/Source/Runtime/ActorAnimation/Public/ActorAnimationPlayer.h b/Engine/Source/Runtime/ActorAnimation/Public/ActorAnimationPlayer.h
--- a/Engine/Source/Runtime/ActorAnimation/Public/ActorAnimationPlayer.h
+++ b/Engine/Source/Runtime/ActorAnimation/Public/ActorAnimationPlayer.h
@@ -16,8 +16,8 @@ USTRUCT(BlueprintType)
 struct FActorAnimationPlaybackSettings
 {
 	FActorAnimationPlaybackSettings()
-		: LoopCount(-1)
-		, PlaySpeed(1.f)
+		: LoopCount(0)
+		, PlayRate(1.f)
 	{}
 
 	GENERATED_BODY()
@@ -26,9 +26,9 @@ struct FActorAnimationPlaybackSettings
 	UPROPERTY(EditAnywhere, BlueprintReadWrite, Category="Playback", meta=(UIMin=-1))
 	int32 LoopCount;
 
-	/** The speed at which to playback the animation */
+	/** The rate at which to playback the animation */
 	UPROPERTY(EditAnywhere, BlueprintReadWrite, Category="Playback")
-	float PlaySpeed;
+	float PlayRate;
 };
 
 /**
@@ -66,12 +66,19 @@ public:
 	 * @param Settings The desired playback settings
 	 */
 	UFUNCTION(BlueprintCallable, Category="Game|Cinematic", meta=(WorldContext="WorldContextObject"))
-	static UActorAnimationPlayer* CreateActorAnimationPlayer(UObject* WorldContextObject, UActorAnimation* ActorAnimation, const FActorAnimationPlaybackSettings& Settings);
+	static UActorAnimationPlayer* CreateActorAnimationPlayer(UObject* WorldContextObject, UActorAnimation* ActorAnimation, FActorAnimationPlaybackSettings Settings);
 
-	/** Start playback from the current time cursor position. */
+	/** Start playback from the current time cursor position, using the current play rate. */
 	UFUNCTION(BlueprintCallable, Category="Game|Cinematic")
 	void Play();
 
+	/**
+	 * Start playback from the current time cursor position, looping the specified number of times.
+	 * @param NumLoops - The number of loops to play. -1 indicates infinite looping.
+	 */
+	UFUNCTION(BlueprintCallable, Category="Game|Cinematic")
+	void PlayLooping(int32 NumLoops = -1);
+
 	/** Pause playback. */
 	UFUNCTION(BlueprintCallable, Category="Game|Cinematic")
 	void Pause();
@@ -84,7 +91,11 @@ public:
 	UFUNCTION(BlueprintCallable, Category="Game|Cinematic")
 	float GetPlaybackPosition() const;
 
-	/** Set the current playback position */
+	/**
+	 * Set the current playback position
+	 * @param NewPlaybackPosition - The new playback position to set.
+	 * If the animation is currently playing, it will continue to do so from the new position
+	 */
 	UFUNCTION(BlueprintCallable, Category="Game|Cinematic")
 	void SetPlaybackPosition(float NewPlaybackPosition);
 
@@ -96,6 +107,17 @@ public:
 	UFUNCTION(BlueprintCallable, Category="Game|Cinematic")
 	float GetLength() const;
 
+	/** Get the playback rate of this player. */
+	UFUNCTION(BlueprintCallable, Category="Game|Cinematic")
+	float GetPlayRate() const;
+
+	/**
+	 * Set the playback rate of this player. Negative values will play the animation in reverse.
+	 * @param PlayRate - The new rate of playback for the animation.
+	 */
+	UFUNCTION(BlueprintCallable, Category="Game|Cinematic")
+	void SetPlayRate(float PlayRate);
+
 protected:
 
 	// IMovieScenePlayer interface
diff --git a/Engine/Source/Runtime/ActorAnimation/Public/MovieSceneActor.h b/Engine/Source/Runtime/ActorAnimation/Public/MovieSceneActor.h
--- a/Engine/Source/Runtime/ActorAnimation/Public/MovieSceneActor.h
+++ b/Engine/Source/Runtime/ActorAnimation/Public/MovieSceneActor.h
@@ -16,16 +16,16 @@ public:
 
 	GENERATED_BODY()
 
-	UPROPERTY(EditAnywhere, BlueprintReadWrite, Category="Playback")
+	UPROPERTY(EditAnywhere, BlueprintReadOnly, Category="Playback")
 	bool bAutoPlay;
 
-	UPROPERTY(EditAnywhere, BlueprintReadWrite, Category="Playback", meta=(ShowOnlyInnerProperties))
+	UPROPERTY(EditAnywhere, BlueprintReadOnly, Category="Playback", meta=(ShowOnlyInnerProperties))
 	FActorAnimationPlaybackSettings PlaybackSettings;
 
 	UPROPERTY(transient, BlueprintReadOnly, Category="Playback")
 	UActorAnimationPlayer* AnimationPlayer;
 
-	UPROPERTY(EditAnywhere, BlueprintReadWrite, Category="General", meta=(AllowedClasses="ActorAnimation"))
+	UPROPERTY(EditAnywhere, BlueprintReadOnly, Category="General", meta=(AllowedClasses="ActorAnimation"))
 	FStringAssetReference ActorAnimation;
 
 public:
