MovieSceneActor improvements:
    - Actor Animation instance is now updated on PostLoad (fixes actor not working in some cases)
    - Added ability to edit associated asset from the level editor

diff --git a/Engine/Source/Runtime/ActorAnimation/Private/MovieSceneActor.cpp b/Engine/Source/Runtime/ActorAnimation/Private/MovieSceneActor.cpp
--- a/Engine/Source/Runtime/ActorAnimation/Private/MovieSceneActor.cpp
+++ b/Engine/Source/Runtime/ActorAnimation/Private/MovieSceneActor.cpp
@@ -48,11 +48,47 @@ void AMovieSceneActor::PostInitProperties()
 	Super::PostInitProperties();
 
 	if (!HasAnyFlags(RF_ClassDefaultObject))
+	{
+		// Make an instance of our asset so that we can keep hard references to actors we are using
+		AnimationInstance = NewObject<UActorAnimationInstance>(this, "AnimationInstance");
+	}
+}
+
+void AMovieSceneActor::PostLoad()
+{
+	Super::PostLoad();
+	
+	UpdateAnimationInstance();
+}
+
+#if WITH_EDITOR
+void AMovieSceneActor::PostEditChangeProperty( struct FPropertyChangedEvent& PropertyChangedEvent)
+{
+	FName PropertyName = PropertyChangedEvent.Property ? PropertyChangedEvent.Property->GetFName() : NAME_None;
+
+	if (PropertyName == GET_MEMBER_NAME_CHECKED(AMovieSceneActor, ActorAnimation))
 	{
 		UpdateAnimationInstance();
 	}
+
+	Super::PostEditChangeProperty(PropertyChangedEvent);
+}
+
+bool AMovieSceneActor::GetReferencedContentObjects(TArray<UObject*>& Objects) const
+{
+	if (UObject* Asset = AnimationInstance ? AnimationInstance->GetActorAnimation() : nullptr)
+	{
+		// @todo: Enable editing of animation instances
+		Objects.Add(Asset);
+	}
+
+	Super::GetReferencedContentObjects(Objects);
+
+	return true;
 }
 
+#endif // WITH_EDITOR
+
 void AMovieSceneActor::Tick(float DeltaSeconds)
 {
 	if (AnimationPlayer)
@@ -89,11 +125,5 @@ void AMovieSceneActor::InitializePlayer()
 void AMovieSceneActor::UpdateAnimationInstance()
 {
 	UActorAnimation* ActorAnimationAsset = Cast<UActorAnimation>(ActorAnimation.TryLoad());
-	if (!AnimationInstance)
-	{
-		// Make an instance of our asset so that we can keep hard references to actors we are using
-		AnimationInstance = NewObject<UActorAnimationInstance>(this, "AnimationInstance");
-	}
-
 	AnimationInstance->Initialize(ActorAnimationAsset, GetWorld(), true);
 }
\ No newline at end of file
diff --git a/Engine/Source/Runtime/ActorAnimation/Public/ActorAnimation.h b/Engine/Source/Runtime/ActorAnimation/Public/ActorAnimation.h
--- a/Engine/Source/Runtime/ActorAnimation/Public/ActorAnimation.h
+++ b/Engine/Source/Runtime/ActorAnimation/Public/ActorAnimation.h
@@ -68,6 +68,13 @@ public:
 	 */
 	void SetContext(UObject* NewContext);
 
+	/**
+	 * Get the actor animation of which this is an instance
+	 *
+	 * @return this instance's parent actor animation
+	 */
+	UActorAnimation* GetActorAnimation() { return ActorAnimation; }
+
 public:
 
 	// UMovieSceneSequence interface
diff --git a/Engine/Source/Runtime/ActorAnimation/Public/MovieSceneActor.h b/Engine/Source/Runtime/ActorAnimation/Public/MovieSceneActor.h
--- a/Engine/Source/Runtime/ActorAnimation/Public/MovieSceneActor.h
+++ b/Engine/Source/Runtime/ActorAnimation/Public/MovieSceneActor.h
@@ -38,6 +38,11 @@ public:
 	virtual void BeginPlay() override;
 	virtual void Tick(float DeltaSeconds) override;
 	virtual void PostInitProperties() override;
+	virtual void PostLoad() override;
+#if WITH_EDITOR
+	virtual void PostEditChangeProperty( struct FPropertyChangedEvent& PropertyChangedEvent) override;
+	virtual bool GetReferencedContentObjects(TArray<UObject*>& Objects) const override;
+#endif		// WITH_EDITOR
 
 	void InitializePlayer();
 
