Sequencer: Remove guid from cached object bindings when unpossessing an object. This fixes a bug where if you remove an object binding and then add the object again, there will be two object bindings, resulting in two nodes in the sequencer tree view.

diff --git a/Engine/Source/Runtime/ActorAnimation/Private/ActorAnimation.cpp b/Engine/Source/Runtime/ActorAnimation/Private/ActorAnimation.cpp
--- a/Engine/Source/Runtime/ActorAnimation/Private/ActorAnimation.cpp
+++ b/Engine/Source/Runtime/ActorAnimation/Private/ActorAnimation.cpp
@@ -373,6 +373,8 @@ void UActorAnimationInstance::UnbindPossessableObjects(const FGuid& ObjectId)
 	{
 		ActorAnimation->DefaultObjectReferences.RemoveBinding(ObjectId);
 	}
+
+	CachedObjectBindings.Remove(ObjectId);
 }
 
 void UActorAnimationInstance::DeselectAllActors()
