Sequencer: Missing changes.

diff --git a/Engine/Source/Runtime/ActorAnimation/Private/ActorAnimationPlayer.cpp b/Engine/Source/Runtime/ActorAnimation/Private/ActorAnimationPlayer.cpp
--- a/Engine/Source/Runtime/ActorAnimation/Private/ActorAnimationPlayer.cpp
+++ b/Engine/Source/Runtime/ActorAnimation/Private/ActorAnimationPlayer.cpp
@@ -368,6 +368,13 @@ void UActorAnimationPlayer::UpdateCameraCut(UObject* ObjectToViewThrough, bool b
 {
 }
 
+void UActorAnimationPlayer::SetViewportSettings(const TMap<FViewportClient*, EMovieSceneViewportParams>& ViewportParamsMap)
+{
+}
+
+void UActorAnimationPlayer::GetViewportSettings(TMap<FViewportClient*, EMovieSceneViewportParams>& ViewportParamsMap) const
+{
+}
 
 EMovieScenePlayerStatus::Type UActorAnimationPlayer::GetPlaybackStatus() const
 {
